Default to spline with auto tangents
Remove minor ticks
Misc fixes
Expose basic BP API

diff --git a/Engine/Source/Editor/ComponentVisualizers/Private/SplineComponentVisualizer.cpp b/Engine/Source/Editor/ComponentVisualizers/Private/SplineComponentVisualizer.cpp
--- a/Engine/Source/Editor/ComponentVisualizers/Private/SplineComponentVisualizer.cpp
+++ b/Engine/Source/Editor/ComponentVisualizers/Private/SplineComponentVisualizer.cpp
@@ -11,7 +11,7 @@ void FSplineComponentVisualizer::DrawVisualization(const UActorComponent* Compon
 {
 	if (const USplineComponent* SplineComp = Cast<const USplineComponent>(Component))
 	{
-		FColor DrawColor = SplineComp->SplineColor;
+		FColor DrawColor(255,255,255);
 
 		const FInterpCurveVector& SplineInfo = SplineComp->SplineInfo;
 
@@ -43,7 +43,7 @@ void FSplineComponentVisualizer::DrawVisualization(const UActorComponent* Compon
 				}
 				else
 				{
-					int32 NumSteps = FMath::CeilToInt((NewKeyTime - OldKeyTime) / 0.1f);
+					int32 NumSteps = FMath::CeilToInt((NewKeyTime - OldKeyTime) / 0.02f);
 					float DrawSubstep = (NewKeyTime - OldKeyTime) / NumSteps;
 
 					// Find position on first keyframe.
@@ -59,7 +59,7 @@ void FSplineComponentVisualizer::DrawVisualization(const UActorComponent* Compon
 						PDI->DrawLine(OldPos, NewPos, DrawColor, SDPG_World);
 
 						// Don't draw point for last one - its the keypoint drawn above.
-						if (StepIdx != NumSteps)
+						if (false)// (StepIdx != NumSteps)
 						{
 							PDI->DrawPoint(NewPos, DrawColor, 3.f, SDPG_World);
 						}
@@ -138,7 +138,8 @@ bool FSplineComponentVisualizer::HandleInputDelta(FLevelEditorViewportClient* Vi
 		{
 			// Duplicate key
 			FInterpCurvePoint<FVector> KeyToDupe = SplineComp->SplineInfo.Points[SelectedKeyIndex];
-			SplineComp->SplineInfo.Points.Insert(KeyToDupe, SelectedKeyIndex);
+			KeyToDupe.InterpMode = CIM_CurveAuto;
+			int32 NewKeyIndex = SplineComp->SplineInfo.Points.Insert(KeyToDupe, SelectedKeyIndex);
 			// move selection to 'next' key
 			SelectedKeyIndex++;
 
@@ -156,6 +157,9 @@ bool FSplineComponentVisualizer::HandleInputDelta(FLevelEditorViewportClient* Vi
 		// Convert back to local space
 		SplineComp->SplineInfo.Points[SelectedKeyIndex].OutVal = SplineComp->ComponentToWorld.InverseTransformPosition(NewWorldPos);
 
+		// Update tangents
+		SplineComp->SplineInfo.AutoSetTangents();
+
 		return true;
 	}
 
@@ -181,7 +185,13 @@ bool FSplineComponentVisualizer::HandleInputKey(FLevelEditorViewportClient* View
 		{
 			SplineComp->SplineInfo.Points.RemoveAt(SelectedKeyIndex);
 
-			bHandled = true;
+			SplineComp->RefreshSplineInputs(); // update input value for each key
+
+			SplineComp->SplineInfo.AutoSetTangents(); // update tangents
+
+			SelectedKeyIndex = INDEX_NONE; // deselect any keys
+
+			bHandled = true; // consume key input
 		}
 	}
 
diff --git a/Engine/Source/Runtime/Engine/Classes/Components/SplineComponent.h b/Engine/Source/Runtime/Engine/Classes/Components/SplineComponent.h
--- a/Engine/Source/Runtime/Engine/Classes/Components/SplineComponent.h
+++ b/Engine/Source/Runtime/Engine/Classes/Components/SplineComponent.h
@@ -27,19 +27,15 @@ public:
 	FInterpCurveVector SplineInfo;
 };
 
-UCLASS()
+UCLASS(ClassGroup=Shapes, meta=(BlueprintSpawnableComponent))
 class ENGINE_API USplineComponent : public USceneComponent
 {
 	GENERATED_UCLASS_BODY()
 
 	/** Actual data for spline. Locations and tangents are in world space. */
-	UPROPERTY(EditAnywhere, BlueprintReadWrite, Category=SplineComponent)
+	UPROPERTY()
 	FInterpCurveVector SplineInfo;
 
-	/** Color of spline */
-	UPROPERTY(EditAnywhere, BlueprintReadWrite, Category=SplineComponent)
-	FColor SplineColor;
-
 	/** Input, distance along curve, output, parameter that puts you there. */
 	UPROPERTY()
 	FInterpCurveFloat SplineReparamTable;
@@ -53,12 +49,15 @@ class ENGINE_API USplineComponent : public USceneComponent
 	void UpdateSplineReparamTable();
 
 	/** Returns total length along this spline */
+	UFUNCTION(BlueprintCallable, Category=Spline) 
 	float GetSplineLength() const;
 	
 	/** Given a distance along the length of this spline, return the point in space where this puts you */
+	UFUNCTION(BlueprintCallable, Category=Spline)
 	FVector GetLocationAtDistanceAlongSpline(float Distance) const;
 	
 	/** Given a distance along the length of this spline, return the direction of the spline there. Note, result is non-unit length. */
+	UFUNCTION(BlueprintCallable, Category=Spline)
 	FVector GetTangentAtDistanceAlongSpline(float Distance) const;
 
 	/** Walk through keys and set time for each one */
diff --git a/Engine/Source/Runtime/Engine/Private/SplineComponent.cpp b/Engine/Source/Runtime/Engine/Private/SplineComponent.cpp
--- a/Engine/Source/Runtime/Engine/Private/SplineComponent.cpp
+++ b/Engine/Source/Runtime/Engine/Private/SplineComponent.cpp
@@ -10,16 +10,14 @@
 USplineComponent::USplineComponent(const class FPostConstructInitializeProperties& PCIP)
 	: Super(PCIP)
 {
-	SplineColor = FColor(255,255,255);
-
 	// Add 2 keys by default
 	int32 PointIndex = SplineInfo.AddPoint(0.f, FVector(0,0,0));
-	FInterpCurvePoint<FVector>& Point = SplineInfo.Points[PointIndex];
-	Point.LeaveTangent = FVector(1,0,0);
+	SplineInfo.Points[PointIndex].InterpMode = CIM_CurveAuto;
 
 	PointIndex = SplineInfo.AddPoint(1.f, FVector(100,0,0));
-	Point = SplineInfo.Points[PointIndex];
-	Point.ArriveTangent = FVector(1, 0, 0);
+	SplineInfo.Points[PointIndex].InterpMode = CIM_CurveAuto;
+
+	SplineInfo.AutoSetTangents();
 }
 
 void USplineComponent::UpdateSplineReparamTable()
