Improve spline remap table resolution
Notify actor of changes as we modify spline

diff --git a/Engine/Source/Editor/ComponentVisualizers/Private/SplineComponentVisualizer.cpp b/Engine/Source/Editor/ComponentVisualizers/Private/SplineComponentVisualizer.cpp
--- a/Engine/Source/Editor/ComponentVisualizers/Private/SplineComponentVisualizer.cpp
+++ b/Engine/Source/Editor/ComponentVisualizers/Private/SplineComponentVisualizer.cpp
@@ -161,6 +161,12 @@ bool FSplineComponentVisualizer::HandleInputDelta(FLevelEditorViewportClient* Vi
 		SplineComp->SplineInfo.AutoSetTangents();
 		SplineComp->UpdateSplineReparamTable();
 
+		// Notify of change so any CS is re-run
+		if (SplineOwningActor.IsValid())
+		{
+			SplineOwningActor.Get()->PostEditMove(true);
+		}
+
 		return true;
 	}
 
@@ -191,6 +197,12 @@ bool FSplineComponentVisualizer::HandleInputKey(FLevelEditorViewportClient* View
 			SplineComp->UpdateSplineReparamTable(); // update reparam table
 			SelectedKeyIndex = INDEX_NONE; // deselect any keys
 
+			// Notify of change so any CS is re-run
+			if(SplineOwningActor.IsValid())
+			{
+				SplineOwningActor.Get()->PostEditMove(true);
+			}
+
 			bHandled = true; // consume key input
 		}
 	}
diff --git a/Engine/Source/Runtime/Engine/Private/SplineComponent.cpp b/Engine/Source/Runtime/Engine/Private/SplineComponent.cpp
--- a/Engine/Source/Runtime/Engine/Private/SplineComponent.cpp
+++ b/Engine/Source/Runtime/Engine/Private/SplineComponent.cpp
@@ -34,11 +34,12 @@ void USplineComponent::UpdateSplineReparamTable()
 		return;
 	}
 	
-	const int32 NumSteps = 10; // TODO: Make this adaptive...
+	const int32 NumKeys = SplineInfo.Points.Num(); // Number of keys on spline
+	const int32 NumSteps = 10 * (NumKeys-1); // TODO: Make this adaptive...
 	
 	// Find range of input
 	float Param = SplineInfo.Points[0].InVal;
-	const float MaxInput = SplineInfo.Points[SplineInfo.Points.Num()-1].InVal;
+	const float MaxInput = SplineInfo.Points[NumKeys-1].InVal;
 	const float Interval = (MaxInput - Param)/((float)(NumSteps-1)); 
 	
 	// Add first entry, using first point on curve, total distance will be 0
