Improvements to SplineComponent and its visualization.
#branch UE4
#proj Editor.ComponentVisualizers, Editor.UnrealEd, Runtime.Core, Runtime.Engine
#change FInterpCurve now uses a binary search to obtain the correct key. Removed the optional out-parameter on Eval, EvalDerivative, etc. Removed  FORCEINLINEs on large methods.
#change USplineComponent now builds the remap table with points spaced at equal distances, which leads to better results.  Added internal methods to obtain the segment length, and the interpolation parameter along the segment at a given distance.
#add Added Duration to USplineComponent.
#add Added a number of different methods for obtaining position, tangent and rotation at a specified distance, or a specified time (optionally at constant velocity).
#add Added methods for building the spline procedurally.
#add Supports right-click context menu on ComponentVisualizers in the standard interface.
#add Key tangents can be manually edited.
#change FLevelEditorViewportClient::ProcessClick now handles right clicks to the Unreal gizmo in the same manner as left clicks, i.e. the method recurses using a version of the hit proxy with the gizmo removed, so we can look at what's underneath.  This allows for "genuine" context menus, i.e. depending on what was clicked, instead of using a single default provided by the current level editor.

diff --git a/Engine/Source/Editor/ComponentVisualizers/ComponentVisualizers.Build.cs b/Engine/Source/Editor/ComponentVisualizers/ComponentVisualizers.Build.cs
--- a/Engine/Source/Editor/ComponentVisualizers/ComponentVisualizers.Build.cs
+++ b/Engine/Source/Editor/ComponentVisualizers/ComponentVisualizers.Build.cs
@@ -15,8 +15,10 @@ public class ComponentVisualizers : ModuleRules
                 "InputCore",
 				"Engine",
 				"Slate",
+                "SlateCore",
 				"UnrealEd",
                 "PropertyEditor",
+                "EditorStyle"
 			}
 			);
 	}
diff --git a/Engine/Source/Editor/ComponentVisualizers/Private/ComponentVisualizers.cpp b/Engine/Source/Editor/ComponentVisualizers/Private/ComponentVisualizers.cpp
--- a/Engine/Source/Editor/ComponentVisualizers/Private/ComponentVisualizers.cpp
+++ b/Engine/Source/Editor/ComponentVisualizers/Private/ComponentVisualizers.cpp
@@ -46,4 +46,9 @@ void FComponentVisualizersModule::RegisterComponentVisualizer(FName ComponentCla
 	}
 
 	RegisteredComponentClassNames.Add(ComponentClassName);
+
+	if (Visualizer.IsValid())
+	{
+		Visualizer->OnRegister();
+	}
 }
diff --git a/Engine/Source/Editor/ComponentVisualizers/Private/ComponentVisualizersPrivatePCH.h b/Engine/Source/Editor/ComponentVisualizers/Private/ComponentVisualizersPrivatePCH.h
--- a/Engine/Source/Editor/ComponentVisualizers/Private/ComponentVisualizersPrivatePCH.h
+++ b/Engine/Source/Editor/ComponentVisualizers/Private/ComponentVisualizersPrivatePCH.h
@@ -4,6 +4,7 @@
 #define __ComponentVisualizersPrivatePCH_h__
 
 #include "UnrealEd.h"
+#include "Slate.h"
 #include "Editor/PropertyEditor/Public/PropertyEditing.h"
 
 #endif // __ComponentVisualizersPrivatePCH_h__
diff --git a/Engine/Source/Editor/ComponentVisualizers/Private/SplineComponentVisualizer.cpp b/Engine/Source/Editor/ComponentVisualizers/Private/SplineComponentVisualizer.cpp
--- a/Engine/Source/Editor/ComponentVisualizers/Private/SplineComponentVisualizer.cpp
+++ b/Engine/Source/Editor/ComponentVisualizers/Private/SplineComponentVisualizer.cpp
@@ -6,6 +6,77 @@
 IMPLEMENT_HIT_PROXY(HSplineVisProxy, HComponentVisProxy);
 IMPLEMENT_HIT_PROXY(HSplineKeyProxy, HSplineVisProxy);
 
+#define LOCTEXT_NAMESPACE "SplineComponentVisualizer"
+
+
+/** Define commands for the spline component visualizer */
+class FSplineComponentVisualizerCommands : public TCommands<FSplineComponentVisualizerCommands>
+{
+public:
+	FSplineComponentVisualizerCommands() : TCommands <FSplineComponentVisualizerCommands>
+	(
+		"SplineComponentVisualizer",	// Context name for fast lookup
+		LOCTEXT("SplineComponentVisualizer", "Spline Component Visualizer"),	// Localized context name for displaying
+		NAME_None,	// Parent
+		FEditorStyle::GetStyleSetName()
+	)
+	{
+	}
+
+	virtual void RegisterCommands() OVERRIDE
+	{
+		UI_COMMAND(DeleteKey, "Delete key", "Delete the currently selected key.", EUserInterfaceActionType::Button, FInputGesture(EKeys::Delete));
+		UI_COMMAND(DuplicateKey, "Duplicate key", "Duplicates the currently selected key.", EUserInterfaceActionType::Button, FInputGesture());
+		UI_COMMAND(ResetToAutomaticTangent, "Reset to Automatic Tangent", "Specify that the tangent for this key will be reset to its default auto-generated value.", EUserInterfaceActionType::Button, FInputGesture());
+	}
+
+public:
+	/** Delete key */
+	TSharedPtr<FUICommandInfo> DeleteKey;
+
+	/** Duplicate key */
+	TSharedPtr<FUICommandInfo> DuplicateKey;
+
+	/** Reset to automatic tangent */
+	TSharedPtr<FUICommandInfo> ResetToAutomaticTangent;
+};
+
+
+
+FSplineComponentVisualizer::FSplineComponentVisualizer()
+	: FComponentVisualizer()
+	, SelectedKeyIndex(INDEX_NONE)
+	, bAllowDuplication(true)
+{
+	FSplineComponentVisualizerCommands::Register();
+
+	SplineComponentVisualizerActions = MakeShareable(new FUICommandList);
+}
+
+void FSplineComponentVisualizer::OnRegister()
+{
+	const auto& Commands = FSplineComponentVisualizerCommands::Get();
+
+	SplineComponentVisualizerActions->MapAction(
+		Commands.DeleteKey,
+		FExecuteAction::CreateSP(this, &FSplineComponentVisualizer::OnDeleteKey),
+		FCanExecuteAction::CreateSP(this, &FSplineComponentVisualizer::IsSelectionValid));
+
+	SplineComponentVisualizerActions->MapAction(
+		Commands.DuplicateKey,
+		FExecuteAction::CreateSP(this, &FSplineComponentVisualizer::OnDuplicateKey),
+		FCanExecuteAction::CreateSP(this, &FSplineComponentVisualizer::IsSelectionValid));
+
+	SplineComponentVisualizerActions->MapAction(
+		Commands.ResetToAutomaticTangent,
+		FExecuteAction::CreateSP(this, &FSplineComponentVisualizer::OnResetToAutomaticTangent),
+		FCanExecuteAction::CreateSP(this, &FSplineComponentVisualizer::CanResetToAutomaticTangent));
+}
+
+FSplineComponentVisualizer::~FSplineComponentVisualizer()
+{
+	FSplineComponentVisualizerCommands::Unregister();
+}
 
 void FSplineComponentVisualizer::DrawVisualization(const UActorComponent* Component, const FSceneView* View, FPrimitiveDrawInterface* PDI)
 {
@@ -15,6 +86,17 @@ void FSplineComponentVisualizer::DrawVisualization(const UActorComponent* Compon
 
 		const FInterpCurveVector& SplineInfo = SplineComp->SplineInfo;
 
+		/*
+		// Debug draw keys on the SplineReparamTable
+		const FInterpCurveFloat& SplineReparamTable = SplineComp->SplineReparamTable;
+		for (const auto& Point : SplineReparamTable.Points)
+		{
+			const float Value = Point.OutVal;
+			FVector KeyPos = SplineComp->ComponentToWorld.TransformPosition(SplineInfo.Eval(Value, FVector::ZeroVector));
+			PDI->DrawPoint(KeyPos, FColor(255, 0, 255), 5.f, SDPG_World);
+		}
+		*/
+
 		FVector OldKeyPos(0);
 		float OldKeyTime = 0.f;
 		for (int32 KeyIdx = 0; KeyIdx<SplineInfo.Points.Num(); KeyIdx++)
@@ -28,7 +110,7 @@ void FSplineComponentVisualizer::DrawVisualization(const UActorComponent* Compon
 			// Draw the keypoint
 			PDI->SetHitProxy(new HSplineKeyProxy(Component, KeyIdx));
 
-			PDI->DrawPoint(NewKeyPos, KeyColor, 6.f, SDPG_World);
+			PDI->DrawPoint(NewKeyPos, KeyColor, 8.f, SDPG_World);
 
 			PDI->SetHitProxy(NULL);
 
@@ -105,13 +187,13 @@ bool FSplineComponentVisualizer::VisProxyHandleClick(HComponentVisProxy* VisProx
 	return bEditing;
 }
 
-USplineComponent* FSplineComponentVisualizer::GetEditedSplineComponent()
+USplineComponent* FSplineComponentVisualizer::GetEditedSplineComponent() const
 {
 	return Cast<USplineComponent>(GetComponentFromPropertyName(SplineOwningActor.Get(), SplineCompPropName));
 }
 
 
-bool FSplineComponentVisualizer::GetWidgetLocation(FVector& OutLocation)
+bool FSplineComponentVisualizer::GetWidgetLocation(FVector& OutLocation) const
 {
 	USplineComponent* SplineComp = GetEditedSplineComponent();
 	if (SplineComp && SelectedKeyIndex != INDEX_NONE)
@@ -127,7 +209,7 @@ bool FSplineComponentVisualizer::GetWidgetLocation(FVector& OutLocation)
 }
 
 
-bool FSplineComponentVisualizer::HandleInputDelta(FLevelEditorViewportClient* ViewportClient, FViewport* Viewport, FVector& DeltaTranslate, FRotator& DeltalRotate, FVector& DeltaScale)
+bool FSplineComponentVisualizer::HandleInputDelta(FLevelEditorViewportClient* ViewportClient, FViewport* Viewport, FVector& DeltaTranslate, FRotator& DeltaRotate, FVector& DeltaScale)
 {
 	USplineComponent* SplineComp = GetEditedSplineComponent();
 	if( SplineComp && 
@@ -136,26 +218,59 @@ bool FSplineComponentVisualizer::HandleInputDelta(FLevelEditorViewportClient* Vi
 	{
 		if (ViewportClient->IsAltPressed() && bAllowDuplication)
 		{
-			// Duplicate key
-			FInterpCurvePoint<FVector> KeyToDupe = SplineComp->SplineInfo.Points[SelectedKeyIndex];
-			KeyToDupe.InterpMode = CIM_CurveAuto;
-			int32 NewKeyIndex = SplineComp->SplineInfo.Points.Insert(KeyToDupe, SelectedKeyIndex);
-			// move selection to 'next' key
-			SelectedKeyIndex++;
-
-			// Update Input value for all keys
-			SplineComp->RefreshSplineInputs();
-
+			OnDuplicateKey();
 			// Don't duplicate again until we release LMB
 			bAllowDuplication = false;
 		}
 
-		// Find key position in world space
-		FVector CurrentWorldPos = SplineComp->ComponentToWorld.TransformPosition(SplineComp->SplineInfo.Points[SelectedKeyIndex].OutVal);
-		// Move in world space
-		FVector NewWorldPos = CurrentWorldPos + DeltaTranslate;
-		// Convert back to local space
-		SplineComp->SplineInfo.Points[SelectedKeyIndex].OutVal = SplineComp->ComponentToWorld.InverseTransformPosition(NewWorldPos);
+		FInterpCurvePoint<FVector>& EditedPoint = SplineComp->SplineInfo.Points[SelectedKeyIndex];
+
+		if (!DeltaTranslate.IsZero())
+		{
+			// Find key position in world space
+			FVector CurrentWorldPos = SplineComp->ComponentToWorld.TransformPosition(EditedPoint.OutVal);
+			// Move in world space
+			FVector NewWorldPos = CurrentWorldPos + DeltaTranslate;
+			// Convert back to local space
+			EditedPoint.OutVal = SplineComp->ComponentToWorld.InverseTransformPosition(NewWorldPos);
+		}
+
+		if (!DeltaRotate.IsZero())
+		{
+			// Set point tangent as user controlled
+			EditedPoint.InterpMode = CIM_CurveUser;
+
+			// Rotate tangent according to delta rotation
+			const FVector NewTangent = DeltaRotate.RotateVector(EditedPoint.LeaveTangent);
+			EditedPoint.LeaveTangent = NewTangent;
+			EditedPoint.ArriveTangent = NewTangent;
+		}
+
+		if (!DeltaScale.IsZero())
+		{
+			// Set point tangent as user controlled
+			EditedPoint.InterpMode = CIM_CurveUser;
+
+			// Break tangent into direction and length so we can change its scale (the 'tension')
+			// independently of its direction.
+			FVector Direction;
+			float Length;
+			EditedPoint.LeaveTangent.ToDirectionAndLength(Direction, Length);
+
+			// Figure out which component has changed, and use it
+			float DeltaScaleValue = (DeltaScale.X != 0.0f) ? DeltaScale.X : ((DeltaScale.Y != 0.0f) ? DeltaScale.Y : DeltaScale.Z);
+
+			// Change scale, avoiding singularity by never allowing a scale of 0, hence preserving direction.
+			Length += DeltaScaleValue * 10.0f;
+			if (Length == 0.0f)
+			{
+				Length = SMALL_NUMBER;
+			}
+
+			const FVector NewTangent = Direction * Length;
+			EditedPoint.LeaveTangent = NewTangent;
+			EditedPoint.ArriveTangent = NewTangent;
+		}
 
 		// Update tangents and reparam table
 		SplineComp->SplineInfo.AutoSetTangents();
@@ -175,6 +290,8 @@ bool FSplineComponentVisualizer::HandleInputDelta(FLevelEditorViewportClient* Vi
 
 bool FSplineComponentVisualizer::HandleInputKey(FLevelEditorViewportClient* ViewportClient, FViewport* Viewport, FKey Key, EInputEvent Event)
 {
+	// TODO: find a way to make this work with ProcessCommandBindings(FKeyboardEvent&)
+
 	bool bHandled = false;
 
 	if(Key == EKeys::LeftMouseButton && Event == IE_Released)
@@ -185,24 +302,9 @@ bool FSplineComponentVisualizer::HandleInputKey(FLevelEditorViewportClient* View
 	else if(Key == EKeys::Delete && Event == IE_Pressed)
 	{
 		// Delete selected key
-		USplineComponent* SplineComp = GetEditedSplineComponent();
-		if (SplineComp &&
-			SelectedKeyIndex != INDEX_NONE &&
-			SelectedKeyIndex < SplineComp->SplineInfo.Points.Num())
+		if (IsSelectionValid())
 		{
-			SplineComp->SplineInfo.Points.RemoveAt(SelectedKeyIndex);
-
-			SplineComp->RefreshSplineInputs(); // update input value for each key
-			SplineComp->SplineInfo.AutoSetTangents(); // update tangents
-			SplineComp->UpdateSplineReparamTable(); // update reparam table
-			SelectedKeyIndex = INDEX_NONE; // deselect any keys
-
-			// Notify of change so any CS is re-run
-			if(SplineOwningActor.IsValid())
-			{
-				SplineOwningActor.Get()->PostEditMove(true);
-			}
-
+			OnDeleteKey();
 			bHandled = true; // consume key input
 		}
 	}
@@ -216,4 +318,97 @@ void FSplineComponentVisualizer::EndEditing()
 	SplineOwningActor = NULL;
 	SplineCompPropName = NAME_None;
 	SelectedKeyIndex = INDEX_NONE;
-}
\ No newline at end of file
+}
+
+
+void FSplineComponentVisualizer::OnDuplicateKey()
+{
+	USplineComponent* SplineComp = GetEditedSplineComponent();
+
+	FInterpCurvePoint<FVector> KeyToDupe = SplineComp->SplineInfo.Points[SelectedKeyIndex];
+	KeyToDupe.InterpMode = CIM_CurveAuto;
+	int32 NewKeyIndex = SplineComp->SplineInfo.Points.Insert(KeyToDupe, SelectedKeyIndex);
+	// move selection to 'next' key
+	SelectedKeyIndex++;
+
+	// Update Input value for all keys
+	SplineComp->RefreshSplineInputs();
+	SplineComp->SplineInfo.AutoSetTangents(); // update tangents
+	SplineComp->UpdateSplineReparamTable(); // update reparam table
+
+	GEditor->RedrawLevelEditingViewports(true);
+}
+
+void FSplineComponentVisualizer::OnDeleteKey()
+{
+	USplineComponent* SplineComp = GetEditedSplineComponent();
+
+	SplineComp->SplineInfo.Points.RemoveAt(SelectedKeyIndex);
+
+	SplineComp->RefreshSplineInputs(); // update input value for each key
+	SplineComp->SplineInfo.AutoSetTangents(); // update tangents
+	SplineComp->UpdateSplineReparamTable(); // update reparam table
+	SelectedKeyIndex = INDEX_NONE; // deselect any keys
+
+	// Notify of change so any CS is re-run
+	if (SplineOwningActor.IsValid())
+	{
+		SplineOwningActor.Get()->PostEditMove(true);
+	}
+
+	GEditor->RedrawLevelEditingViewports(true);
+}
+
+bool FSplineComponentVisualizer::IsSelectionValid() const
+{
+	USplineComponent* SplineComp = GetEditedSplineComponent();
+	return (SplineComp &&
+		SelectedKeyIndex != INDEX_NONE &&
+		SelectedKeyIndex < SplineComp->SplineInfo.Points.Num());
+}
+
+
+void FSplineComponentVisualizer::OnResetToAutomaticTangent()
+{
+	USplineComponent* SplineComp = GetEditedSplineComponent();
+	SplineComp->SplineInfo.Points[SelectedKeyIndex].InterpMode = CIM_CurveAuto;
+	SplineComp->SplineInfo.AutoSetTangents();
+	SplineComp->UpdateSplineReparamTable();
+
+	// Notify of change so any CS is re-run
+	if (SplineOwningActor.IsValid())
+	{
+		SplineOwningActor.Get()->PostEditMove(true);
+	}
+
+	GEditor->RedrawLevelEditingViewports(true);
+}
+
+bool FSplineComponentVisualizer::CanResetToAutomaticTangent() const
+{
+	USplineComponent* SplineComp = GetEditedSplineComponent();
+	return (SplineComp &&
+		SelectedKeyIndex != INDEX_NONE &&
+		SelectedKeyIndex < SplineComp->SplineInfo.Points.Num() &&
+		SplineComp->SplineInfo.Points[SelectedKeyIndex].InterpMode != CIM_CurveAuto);
+}
+
+
+TSharedPtr<SWidget> FSplineComponentVisualizer::GenerateContextMenu() const
+{
+	FMenuBuilder MenuBuilder(true, SplineComponentVisualizerActions);
+	{
+		MenuBuilder.BeginSection("SplineKeyEdit");
+		{
+			MenuBuilder.AddMenuEntry(FSplineComponentVisualizerCommands::Get().DeleteKey);
+			MenuBuilder.AddMenuEntry(FSplineComponentVisualizerCommands::Get().DuplicateKey);
+			MenuBuilder.AddMenuEntry(FSplineComponentVisualizerCommands::Get().ResetToAutomaticTangent);
+		}
+		MenuBuilder.EndSection();
+	}
+
+	TSharedPtr<SWidget> MenuWidget = MenuBuilder.MakeWidget();
+	return MenuWidget;
+}
+
+#undef LOCTEXT_NAMESPACE
diff --git a/Engine/Source/Editor/ComponentVisualizers/Private/SplineComponentVisualizer.h b/Engine/Source/Editor/ComponentVisualizers/Private/SplineComponentVisualizer.h
--- a/Engine/Source/Editor/ComponentVisualizers/Private/SplineComponentVisualizer.h
+++ b/Engine/Source/Editor/ComponentVisualizers/Private/SplineComponentVisualizer.h
@@ -33,26 +33,35 @@ struct HSplineKeyProxy : public HSplineVisProxy
 class FSplineComponentVisualizer : public FComponentVisualizer
 {
 public:
-	FSplineComponentVisualizer()
-	: FComponentVisualizer()
-	, SelectedKeyIndex(INDEX_NONE)
-	, bAllowDuplication(true)
-	{}
+	FSplineComponentVisualizer();
+	virtual ~FSplineComponentVisualizer();
 
 	// Begin FComponentVisualizer interface
+	virtual void OnRegister() OVERRIDE;
 	virtual void DrawVisualization(const UActorComponent* Component, const FSceneView* View, FPrimitiveDrawInterface* PDI) OVERRIDE;
 	virtual bool VisProxyHandleClick(HComponentVisProxy* VisProxy) OVERRIDE;
 	virtual void EndEditing() OVERRIDE;
-	virtual bool GetWidgetLocation(FVector& OutLocation) OVERRIDE;
-	virtual bool HandleInputDelta(FLevelEditorViewportClient* ViewportClient, FViewport* Viewport, FVector& DeltaTranslate, FRotator& DeltalRotate, FVector& DeltaScale) OVERRIDE;
+	virtual bool GetWidgetLocation(FVector& OutLocation) const OVERRIDE;
+	virtual bool HandleInputDelta(FLevelEditorViewportClient* ViewportClient, FViewport* Viewport, FVector& DeltaTranslate, FRotator& DeltaRotate, FVector& DeltaScale) OVERRIDE;
 	virtual bool HandleInputKey(FLevelEditorViewportClient* ViewportClient, FViewport* Viewport, FKey Key, EInputEvent Event) OVERRIDE;
+	virtual TSharedPtr<SWidget> GenerateContextMenu() const OVERRIDE;
 	// End FComponentVisualizer interface
 
 	/** Get the spline component we are currently editing */
-	USplineComponent* GetEditedSplineComponent();
+	USplineComponent* GetEditedSplineComponent() const;
 
 private:
 
+	void OnDeleteKey();
+	void OnDuplicateKey();
+	bool IsSelectionValid() const;
+
+	void OnResetToAutomaticTangent();
+	bool CanResetToAutomaticTangent() const;
+
+	/** Output log commands */
+	TSharedPtr<FUICommandList> SplineComponentVisualizerActions;
+
 	/** Actor that owns the currently edited spline */
 	TWeakObjectPtr<AActor> SplineOwningActor;
 	/** Name of property on the actor that references the spline we are editing */
diff --git a/Engine/Source/Editor/UnrealEd/Private/ComponentVisualizerManager.cpp b/Engine/Source/Editor/UnrealEd/Private/ComponentVisualizerManager.cpp
--- a/Engine/Source/Editor/UnrealEd/Private/ComponentVisualizerManager.cpp
+++ b/Engine/Source/Editor/UnrealEd/Private/ComponentVisualizerManager.cpp
@@ -3,7 +3,7 @@
 #include "UnrealEd.h"
 #include "ComponentVisualizerManager.h"
 
-void FComponentVisualizerManager::HandleProxyForComponentVis(HHitProxy *HitProxy)
+bool FComponentVisualizerManager::HandleProxyForComponentVis(HHitProxy *HitProxy)
 {
 	if (HitProxy->IsA(HComponentVisProxy::StaticGetType()))
 	{
@@ -24,6 +24,7 @@ void FComponentVisualizerManager::HandleProxyForComponentVis(HHitProxy *HitProxy
 					}
 
 					EditedVisualizer = Visualizer;
+					return true;
 				}
 			}
 		}
@@ -32,6 +33,8 @@ void FComponentVisualizerManager::HandleProxyForComponentVis(HHitProxy *HitProxy
 	{
 		ClearActiveComponentVis();
 	}
+
+	return false;
 }
 
 void FComponentVisualizerManager::ClearActiveComponentVis()
@@ -43,7 +46,7 @@ void FComponentVisualizerManager::ClearActiveComponentVis()
 	}
 }
 
-bool FComponentVisualizerManager::HandleInputKey(FLevelEditorViewportClient* ViewportClient, FViewport* Viewport, FKey Key, EInputEvent Event)
+bool FComponentVisualizerManager::HandleInputKey(FLevelEditorViewportClient* ViewportClient, FViewport* Viewport, FKey Key, EInputEvent Event) const
 {
 	if(EditedVisualizer.IsValid())
 	{
@@ -56,7 +59,7 @@ bool FComponentVisualizerManager::HandleInputKey(FLevelEditorViewportClient* Vie
 	return false;
 }
 
-bool FComponentVisualizerManager::HandleInputDelta(FLevelEditorViewportClient* InViewportClient, FViewport* InViewport, FVector& InDrag, FRotator& InRot, FVector& InScale)
+bool FComponentVisualizerManager::HandleInputDelta(FLevelEditorViewportClient* InViewportClient, FViewport* InViewport, FVector& InDrag, FRotator& InRot, FVector& InScale) const
 {
 	if (EditedVisualizer.IsValid() && InViewportClient->GetCurrentWidgetAxis() != EAxisList::None)
 	{
@@ -70,7 +73,7 @@ bool FComponentVisualizerManager::HandleInputDelta(FLevelEditorViewportClient* I
 }
 
 
-bool FComponentVisualizerManager::GetWidgetLocation(FVector& OutLocation)
+bool FComponentVisualizerManager::GetWidgetLocation(FVector& OutLocation) const
 {
 	if (EditedVisualizer.IsValid())
 	{
@@ -79,3 +82,14 @@ bool FComponentVisualizerManager::GetWidgetLocation(FVector& OutLocation)
 
 	return false;
 }
+
+
+TSharedPtr<SWidget> FComponentVisualizerManager::GenerateContextMenuForComponentVis() const
+{
+	if (EditedVisualizer.IsValid())
+	{
+		return EditedVisualizer.Pin()->GenerateContextMenu();
+	}
+
+	return TSharedPtr<SWidget>();
+}
diff --git a/Engine/Source/Editor/UnrealEd/Private/EditorModes.cpp b/Engine/Source/Editor/UnrealEd/Private/EditorModes.cpp
--- a/Engine/Source/Editor/UnrealEd/Private/EditorModes.cpp
+++ b/Engine/Source/Editor/UnrealEd/Private/EditorModes.cpp
@@ -568,10 +568,26 @@ bool FEdMode::HandleClick(FLevelEditorViewportClient* InViewportClient, HHitProx
 
 	if(HitProxy)
 	{	
-		GUnrealEd->ComponentVisManager.HandleProxyForComponentVis(HitProxy);
+		bool bHandled = GUnrealEd->ComponentVisManager.HandleProxyForComponentVis(HitProxy);
+		if (bHandled && Click.GetKey() == EKeys::RightMouseButton)
+		{
+			TSharedPtr<SWidget> MenuWidget = GUnrealEd->ComponentVisManager.GenerateContextMenuForComponentVis();
+			if (MenuWidget.IsValid())
+			{
+				auto ParentLevelEditorPinned = InViewportClient->ParentLevelEditor.Pin();
+				if (ParentLevelEditorPinned.IsValid())
+				{
+					FSlateApplication::Get().PushMenu(
+						ParentLevelEditorPinned.ToSharedRef(),
+						MenuWidget.ToSharedRef(),
+						FSlateApplication::Get().GetCursorPos(),
+						FPopupTransitionEffect(FPopupTransitionEffect::ContextMenu));
+				}
+			}
+		}
 	}
 
-	return 0;
+	return false;
 }
 
 void FEdMode::Enter()
diff --git a/Engine/Source/Editor/UnrealEd/Private/LevelEditorViewport.cpp b/Engine/Source/Editor/UnrealEd/Private/LevelEditorViewport.cpp
--- a/Engine/Source/Editor/UnrealEd/Private/LevelEditorViewport.cpp
+++ b/Engine/Source/Editor/UnrealEd/Private/LevelEditorViewport.cpp
@@ -2148,51 +2148,43 @@ void FLevelEditorViewportClient::ProcessClick(FSceneView& View, HHitProxy* HitPr
 		}
 		else if( HitProxy->IsA( HWidgetAxis::StaticGetType() ) )
 		{
-			if(Click.GetKey() == EKeys::RightMouseButton)
+			// The user clicked on an axis translation/rotation hit proxy.  However, we want
+			// to find out what's underneath the axis widget.  To do this, we'll need to render
+			// the viewport's hit proxies again, this time *without* the axis widgets!
+
+			// OK, we need to be a bit evil right here.  Basically we want to hijack the ShowFlags
+			// for the scene so we can re-render the hit proxies without any axis widgets.  We'll
+			// store the original ShowFlags and modify them appropriately
+			const bool bOldModeWidgets1 = EngineShowFlags.ModeWidgets;
+			const bool bOldModeWidgets2 = View.Family->EngineShowFlags.ModeWidgets;
+
+			EngineShowFlags.ModeWidgets = 0;
+			FSceneViewFamily* SceneViewFamily = const_cast< FSceneViewFamily* >( View.Family );
+			SceneViewFamily->EngineShowFlags.ModeWidgets = 0;
+			bool bWasWidgetDragging = Widget->IsDragging();
+			Widget->SetDragging(false);
+
+			// Invalidate the hit proxy map so it will be rendered out again when GetHitProxy
+			// is called
+			Viewport->InvalidateHitProxy();
+
+			// This will actually re-render the viewport's hit proxies!
+			HHitProxy* HitProxyWithoutAxisWidgets = Viewport->GetHitProxy( HitX, HitY );
+			if( HitProxyWithoutAxisWidgets != NULL && !HitProxyWithoutAxisWidgets->IsA( HWidgetAxis::StaticGetType() ))
 			{
-				// If this is a right click, always handle as though we're clicking the backdrop
-				ClickHandlers::ClickBackdrop(this, Click);
+				// Try this again, but without the widget this time!
+				ProcessClick( View, HitProxyWithoutAxisWidgets, Key, Event, HitX, HitY );
 			}
-			else
-			{
-				// The user clicked on an axis translation/rotation hit proxy.  However, we want
-				// to find out what's underneath the axis widget.  To do this, we'll need to render
-				// the viewport's hit proxies again, this time *without* the axis widgets!
-
-				// OK, we need to be a bit evil right here.  Basically we want to hijack the ShowFlags
-				// for the scene so we can re-render the hit proxies without any axis widgets.  We'll
-				// store the original ShowFlags and modify them appropriately
-				const bool bOldModeWidgets1 = EngineShowFlags.ModeWidgets;
-				const bool bOldModeWidgets2 = View.Family->EngineShowFlags.ModeWidgets;
-
-				EngineShowFlags.ModeWidgets = 0;
-				FSceneViewFamily* SceneViewFamily = const_cast< FSceneViewFamily* >( View.Family );
-				SceneViewFamily->EngineShowFlags.ModeWidgets = 0;
-				bool bWasWidgetDragging = Widget->IsDragging();
-				Widget->SetDragging(false);
-
-				// Invalidate the hit proxy map so it will be rendered out again when GetHitProxy
-				// is called
-				Viewport->InvalidateHitProxy();
 
-				// This will actually re-render the viewport's hit proxies!
-				HHitProxy* HitProxyWithoutAxisWidgets = Viewport->GetHitProxy( HitX, HitY );
-				if( HitProxyWithoutAxisWidgets != NULL && !HitProxyWithoutAxisWidgets->IsA( HWidgetAxis::StaticGetType() ))
-				{
-					// Try this again, but without the widget this time!
-					ProcessClick( View, HitProxyWithoutAxisWidgets, Key, Event, HitX, HitY );
-				}
-
-				// Undo the evil
-				EngineShowFlags.ModeWidgets = bOldModeWidgets1;
-				SceneViewFamily->EngineShowFlags.ModeWidgets = bOldModeWidgets2;
+			// Undo the evil
+			EngineShowFlags.ModeWidgets = bOldModeWidgets1;
+			SceneViewFamily->EngineShowFlags.ModeWidgets = bOldModeWidgets2;
 
-				Widget->SetDragging(bWasWidgetDragging);
+			Widget->SetDragging(bWasWidgetDragging);
 
-				// Invalidate the hit proxy map again so that it'll be refreshed with the original
-				// scene contents if we need it again later.
-				Viewport->InvalidateHitProxy();
-			}
+			// Invalidate the hit proxy map again so that it'll be refreshed with the original
+			// scene contents if we need it again later.
+			Viewport->InvalidateHitProxy();
 		}
 	}
 }
diff --git a/Engine/Source/Editor/UnrealEd/Public/ComponentVisualizer.h b/Engine/Source/Editor/UnrealEd/Public/ComponentVisualizer.h
--- a/Engine/Source/Editor/UnrealEd/Public/ComponentVisualizer.h
+++ b/Engine/Source/Editor/UnrealEd/Public/ComponentVisualizer.h
@@ -2,6 +2,8 @@
 
 #pragma once
 
+struct FViewportClick;
+
 struct HComponentVisProxy : public HHitProxy
 {
 	DECLARE_HIT_PROXY(UNREALED_API);
@@ -23,6 +25,11 @@ struct HComponentVisProxy : public HHitProxy
 class UNREALED_API FComponentVisualizer : public TSharedFromThis<FComponentVisualizer>
 {
 public:
+	FComponentVisualizer() {}
+	virtual ~FComponentVisualizer() {}
+
+	/** */
+	virtual void OnRegister() {}
 	/** Draw visualization for the supplied component */
 	virtual void DrawVisualization(const UActorComponent* Component, const FSceneView* View, FPrimitiveDrawInterface* PDI) {}
 	/** */
@@ -30,11 +37,13 @@ public:
 	/** */
 	virtual void EndEditing() {}
 	/** */
-	virtual bool GetWidgetLocation(FVector& OutLocation) { return false; }
+	virtual bool GetWidgetLocation(FVector& OutLocation) const { return false; }
 	/** */
 	virtual bool HandleInputDelta(FLevelEditorViewportClient* ViewportClient, FViewport* Viewport, FVector& DeltaTranslate, FRotator& DeltalRotate, FVector& DeltaScale) { return false; }
 	/** */
 	virtual bool HandleInputKey(FLevelEditorViewportClient* ViewportClient,FViewport* Viewport,FKey Key,EInputEvent Event) { return false; }
+	/** */
+	virtual TSharedPtr<SWidget> GenerateContextMenu() const { return TSharedPtr<SWidget>(); }
 
 	/** Find the name of the property that points to this component */
 	static FName GetComponentPropertyName(const UActorComponent* Component);
diff --git a/Engine/Source/Editor/UnrealEd/Public/ComponentVisualizerManager.h b/Engine/Source/Editor/UnrealEd/Public/ComponentVisualizerManager.h
--- a/Engine/Source/Editor/UnrealEd/Public/ComponentVisualizerManager.h
+++ b/Engine/Source/Editor/UnrealEd/Public/ComponentVisualizerManager.h
@@ -4,6 +4,8 @@
 
 #include "ComponentVisualizer.h"
 
+struct FViewportClick;
+
 /** Class that managed active component visualizer and routes input to it */
 class UNREALED_API FComponentVisualizerManager
 {
@@ -13,16 +15,19 @@ public:
 
 
 	/** Activate a component visualizer given a clicked proxy */
-	void HandleProxyForComponentVis(HHitProxy *HitProxy);
+	bool HandleProxyForComponentVis(HHitProxy *HitProxy);
 	/** Clear active component visualizer */
 	void ClearActiveComponentVis();
 
 	/** Pass key input to active visualizer */
-	bool HandleInputKey(FLevelEditorViewportClient* ViewportClient, FViewport* Viewport, FKey Key, EInputEvent Event);
+	bool HandleInputKey(FLevelEditorViewportClient* ViewportClient, FViewport* Viewport, FKey Key, EInputEvent Event) const;
 	/** Pass delta input to active visualizer */
-	bool HandleInputDelta(FLevelEditorViewportClient* InViewportClient, FViewport* InViewport, FVector& InDrag, FRotator& InRot, FVector& InScale);
+	bool HandleInputDelta(FLevelEditorViewportClient* InViewportClient, FViewport* InViewport, FVector& InDrag, FRotator& InRot, FVector& InScale) const;
 	/** Get widget location from active visualizer */
-	bool GetWidgetLocation(FVector& OutLocation);
+	bool GetWidgetLocation(FVector& OutLocation) const;
+
+	/** Generate context menu for the component visualizer */
+	TSharedPtr<SWidget> GenerateContextMenuForComponentVis() const;
 
 private:
 	/** Currently 'active' visualizer that we should pass input to etc */
diff --git a/Engine/Source/Runtime/Core/Public/Math/InterpCurve.h b/Engine/Source/Runtime/Core/Public/Math/InterpCurve.h
--- a/Engine/Source/Runtime/Core/Public/Math/InterpCurve.h
+++ b/Engine/Source/Runtime/Core/Public/Math/InterpCurve.h
@@ -23,17 +23,17 @@ public:
 	 *	Evaluate the output for an arbitary input value. 
 	 *	For inputs outside the range of the keys, the first/last key value is assumed.
 	 */
-	T Eval( const float InVal, const T& Default, int32* PtIdx = NULL ) const;
+	T Eval( const float InVal, const T& Default ) const;
 
 	/** 
 	 *	Evaluate the derivative at a point on the curve.
 	 */
-	T EvalDerivative( const float InVal, const T& Default, int32* PtIdx = NULL ) const;
+	T EvalDerivative( const float InVal, const T& Default ) const;
 
 	/** 
 	 *	Evaluate the second derivative at a point on the curve.
 	 */
-	T EvalSecondDerivative( const float InVal, const T& Default, int32* PtIdx = NULL ) const;
+	T EvalSecondDerivative( const float InVal, const T& Default ) const;
 
 	/** 
 	 * Find the nearest point on spline to the given point.
@@ -77,6 +77,9 @@ public:
 		return Ar;
 	}
 
+private:
+	/** Find the lower index of the two points whose input values bound the supplied input value */
+	int32 GetPointIndexForInputValue(const float InValue) const;
 };
 
 template< class T > 
@@ -118,248 +121,199 @@ FORCEINLINE void FInterpCurve<T>::Reset()
 	Points.Empty();
 }
 
-template< class T > 
-FORCEINLINE T FInterpCurve<T>::Eval( const float InVal, const T& Default, int32* PtIdx ) const
+template< class T >
+int32 FInterpCurve<T>::GetPointIndexForInputValue(const float InValue) const
 {
 	const int32 NumPoints = Points.Num();
 
-	// If no point in curve, return the Default value we passed in.
-	if( NumPoints == 0 )
+	check(NumPoints > 0);
+
+	if (InValue < Points[0].InVal)
 	{
-		if( PtIdx )
-		{
-			*PtIdx = -1;
-		}
-		return Default;
+		return -1;
 	}
 
-	// If only one point, or before the first point in the curve, return the first points value.
-	if( NumPoints < 2 || (InVal <= Points[0].InVal) )
+	if (InValue > Points[NumPoints - 1].InVal)
 	{
-		if( PtIdx )
-		{
-			*PtIdx = 0;
-		}
-		return Points[0].OutVal;
+		return NumPoints;
 	}
 
-	// If beyond the last point in the curve, return its value.
-	if( InVal >= Points[NumPoints-1].InVal )
+	int32 MinIndex = 0;
+	int32 MaxIndex = NumPoints;
+
+	while (MaxIndex - MinIndex > 1)
 	{
-		if( PtIdx )
+		int32 MidIndex = (MinIndex + MaxIndex) / 2;
+
+		if (Points[MidIndex].InVal <= InValue)
+		{
+			MinIndex = MidIndex;
+		}
+		else
 		{
-			*PtIdx = NumPoints - 1;
+			MaxIndex = MidIndex;
 		}
-		return Points[NumPoints-1].OutVal;
 	}
 
-	// Somewhere with curve range - linear search to find value.
-	for( int32 i=1; i<NumPoints; i++ )
-	{	
-		if( InVal < Points[i].InVal )
-		{
-			const float Diff = Points[i].InVal - Points[i-1].InVal;
+	return MinIndex;
+}
 
-			if( Diff > 0.f && Points[i-1].InterpMode != CIM_Constant )
-			{
-				const float Alpha = (InVal - Points[i-1].InVal) / Diff;
+template< class T > 
+T FInterpCurve<T>::Eval( const float InVal, const T& Default ) const
+{
+	const int32 NumPoints = Points.Num();
 
-				if( PtIdx )
-				{
-					*PtIdx = i - 1;
-				}
+	// If no point in curve, return the Default value we passed in.
+	if( NumPoints == 0 )
+	{
+		return Default;
+	}
 
-				if( Points[i-1].InterpMode == CIM_Linear )
-				{
-					return FMath::Lerp( Points[i-1].OutVal, Points[i].OutVal, Alpha );
-				}
-				else
-				{
-					return FMath::CubicInterp( Points[i-1].OutVal, Points[i-1].LeaveTangent * Diff, Points[i].OutVal, Points[i].ArriveTangent * Diff, Alpha );
-				}
-			}
-			else
-			{
-				if( PtIdx )
-				{
-					*PtIdx = i - 1;
-				}
+	// Binary search to find index of lower bound of input value
+	const int32 Index = GetPointIndexForInputValue(InVal);
 
-				return Points[i-1].OutVal;
-			}
-		}
+	// If before the first point, return its value
+	if (Index == -1)
+	{
+		return Points[0].OutVal;
 	}
 
-	// Shouldn't really reach here.
-	if( PtIdx )
+	// If on or beyond the last point, return its value.
+	if (Index >= NumPoints - 1)
 	{
-		*PtIdx = NumPoints - 1;
+		return Points[NumPoints - 1].OutVal;
 	}
 
-	return Points[NumPoints-1].OutVal;
+	// Somewhere within curve range - interpolate.
+	check(Index >= 0 && Index < NumPoints - 1);
+	const auto& PrevPoint = Points[Index];
+	const auto& NextPoint = Points[Index + 1];
+	const float Diff = NextPoint.InVal - PrevPoint.InVal;
+
+	if (Diff > 0.0f && PrevPoint.InterpMode != CIM_Constant)
+	{
+		const float Alpha = (InVal - PrevPoint.InVal) / Diff;
+
+		if (PrevPoint.InterpMode == CIM_Linear)
+		{
+			return FMath::Lerp(PrevPoint.OutVal, NextPoint.OutVal, Alpha);
+		}
+		else
+		{
+			return FMath::CubicInterp(PrevPoint.OutVal, PrevPoint.LeaveTangent * Diff, NextPoint.OutVal, NextPoint.ArriveTangent * Diff, Alpha);
+		}
+	}
+	else
+	{
+		return Points[Index].OutVal;
+	}
 }
 
 template< class T > 
-FORCEINLINE T FInterpCurve<T>::EvalDerivative( const float InVal, const T& Default, int32* PtIdx ) const
+T FInterpCurve<T>::EvalDerivative( const float InVal, const T& Default ) const
 {
 	const int32 NumPoints = Points.Num();
 
 	// If no point in curve, return the Default value we passed in.
-	if( NumPoints == 0 )
+	if (NumPoints == 0)
 	{
-		if( PtIdx )
-		{
-			*PtIdx = -1;
-		}
 		return Default;
 	}
 
-	// If only one point, or before the first point in the curve, return the first points value.
-	if( NumPoints < 2 || (InVal <= Points[0].InVal) )
+	// Binary search to find index of lower bound of input value
+	const int32 Index = GetPointIndexForInputValue(InVal);
+
+	// If before the first point, return its tangent value
+	if (Index == -1)
 	{
-		if( PtIdx )
-		{
-			*PtIdx = 0;
-		}
 		return Points[0].LeaveTangent;
 	}
 
-	// If beyond the last point in the curve, return its value.
-	if( InVal >= Points[NumPoints-1].InVal )
+	// If on or beyond the last point, return its tangent value.
+	if (Index >= NumPoints - 1)
 	{
-		if( PtIdx )
-		{
-			*PtIdx = NumPoints - 1;
-		}
-		return Points[NumPoints-1].ArriveTangent;
+		return Points[NumPoints - 1].ArriveTangent;
 	}
 
-	// Somewhere with curve range - linear search to find value.
-	for( int32 i=1; i<NumPoints; i++ )
-	{	
-		if( InVal < Points[i].InVal )
-		{
-			const float Diff = Points[i].InVal - Points[i-1].InVal;
-
-			if( Diff > 0.f && Points[i-1].InterpMode != CIM_Constant )
-			{
-				if( PtIdx )
-				{
-					*PtIdx = NumPoints - 1;
-				}
-
-				const float Alpha = (InVal - Points[i-1].InVal) / Diff;
+	// Somewhere within curve range - interpolate.
+	check(Index >= 0 && Index < NumPoints - 1);
+	const auto& PrevPoint = Points[Index];
+	const auto& NextPoint = Points[Index + 1];
+	const float Diff = NextPoint.InVal - PrevPoint.InVal;
 
-				if( Points[i-1].InterpMode == CIM_Linear )
-				{
-					return FMath::Lerp( Points[i-1].OutVal, Points[i].OutVal, Alpha );
-				}
-				else
-				{
-					return FMath::CubicInterpDerivative( Points[i-1].OutVal, Points[i-1].LeaveTangent * Diff, Points[i].OutVal, Points[i].ArriveTangent * Diff, Alpha );
-				}
-			}
-			else
-			{
-				if( PtIdx )
-				{
-					*PtIdx = -1;
-				}
+	if (Diff > 0.0f && PrevPoint.InterpMode != CIM_Constant)
+	{
+		const float Alpha = (InVal - PrevPoint.InVal) / Diff;
 
-				return T(0.f);
-			}
+		if (PrevPoint.InterpMode == CIM_Linear)
+		{
+			return NextPoint.OutVal - PrevPoint.OutVal;
+		}
+		else
+		{
+			return FMath::CubicInterpDerivative(PrevPoint.OutVal, PrevPoint.LeaveTangent * Diff, NextPoint.OutVal, NextPoint.ArriveTangent * Diff, Alpha);
 		}
 	}
-
-	if( PtIdx )
+	else
 	{
-		*PtIdx = NumPoints - 1;
+		return Points[Index].OutVal;
 	}
-
-	// Shouldn't really reach here.
-	return Points[NumPoints-1].OutVal;
 }
 
 template< class T > 
-FORCEINLINE T FInterpCurve<T>::EvalSecondDerivative( const float InVal, const T& Default, int32* PtIdx ) const
+T FInterpCurve<T>::EvalSecondDerivative( const float InVal, const T& Default ) const
 {
 	const int32 NumPoints = Points.Num();
 
 	// If no point in curve, return the Default value we passed in.
-	if( NumPoints == 0 )
+	if (NumPoints == 0)
 	{
-		if( PtIdx )
-		{
-			*PtIdx = -1;
-		}
 		return Default;
 	}
 
-	// If only one point, or before the first point in the curve, return the first points value.
-	if( NumPoints < 2 || (InVal <= Points[0].InVal) )
+	// Binary search to find index of lower bound of input value
+	const int32 Index = GetPointIndexForInputValue(InVal);
+
+	// If before the first point, return 0
+	if (Index == -1)
 	{
-		if( PtIdx )
-		{
-			*PtIdx = 0;
-		}
-		return Points[0].OutVal;
+		return T();
 	}
 
-	// If beyond the last point in the curve, return its value.
-	if( InVal >= Points[NumPoints-1].InVal )
+	// If on or beyond the last point, return 0
+	if (Index >= NumPoints)
 	{
-		if( PtIdx )
-		{
-			*PtIdx = NumPoints - 1;
-		}
-		return Points[NumPoints-1].OutVal;
+		return T();
 	}
 
-	// Somewhere with curve range - linear search to find value.
-	for( int32 i=1; i<NumPoints; i++ )
-	{	
-		if( InVal < Points[i].InVal )
-		{
-			const float Diff = Points[i].InVal - Points[i-1].InVal;
+	// Somewhere within curve range - interpolate.
+	check(Index >= 0 && Index < NumPoints - 1);
+	const auto& PrevPoint = Points[Index];
+	const auto& NextPoint = Points[Index + 1];
+	const float Diff = NextPoint.InVal - PrevPoint.InVal;
 
-			if( Diff > 0.f && Points[i-1].InterpMode != CIM_Constant )
-			{
-				if( PtIdx )
-				{
-					*PtIdx = i - 1;
-				}
-				const float Alpha = (InVal - Points[i-1].InVal) / Diff;
+	if (Diff > 0.0f && PrevPoint.InterpMode != CIM_Constant)
+	{
+		const float Alpha = (InVal - PrevPoint.InVal) / Diff;
 
-				if( Points[i-1].InterpMode == CIM_Linear )
-				{
-					return FMath::Lerp( Points[i-1].OutVal, Points[i].OutVal, Alpha );
-				}
-				else
-				{
-					return FMath::CubicInterpSecondDerivative( Points[i-1].OutVal, Points[i-1].LeaveTangent * Diff, Points[i].OutVal, Points[i].ArriveTangent * Diff, Alpha );
-				}
-			}
-			else
-			{
-				if( PtIdx )
-				{
-					*PtIdx = -1;
-				}
-				return Default;
-			}
+		if (PrevPoint.InterpMode == CIM_Linear)
+		{
+			// No change in tangent, return 0.
+			return T();
+		}
+		else
+		{
+			return FMath::CubicInterpSecondDerivative(PrevPoint.OutVal, PrevPoint.LeaveTangent * Diff, NextPoint.OutVal, NextPoint.ArriveTangent * Diff, Alpha);
 		}
 	}
-	if( PtIdx )
+	else
 	{
-		*PtIdx = NumPoints - 1;
+		return Points[Index].OutVal;
 	}
-
-	// Shouldn't really reach here.
-	return Points[NumPoints-1].OutVal;
 }
 
 template< class T > 
-FORCEINLINE float FInterpCurve<T>::InaccurateFindNearest( const T &PointInSpace, float& OutDistanceSq ) const
+float FInterpCurve<T>::InaccurateFindNearest( const T &PointInSpace, float& OutDistanceSq ) const
 {
 	const int32 NumPoints = Points.Num();
 	if(NumPoints > 1)
@@ -388,7 +342,7 @@ FORCEINLINE float FInterpCurve<T>::InaccurateFindNearest( const T &PointInSpace,
 }
 
 template< class T > 
-FORCEINLINE float FInterpCurve<T>::InaccurateFindNearestOnSegment( const T &PointInSpace, int32 PtIdx, float& OutSquaredDistance ) const
+float FInterpCurve<T>::InaccurateFindNearestOnSegment( const T &PointInSpace, int32 PtIdx, float& OutSquaredDistance ) const
 {
 	if( CIM_Constant == Points[PtIdx].InterpMode )
 	{
@@ -470,7 +424,7 @@ FORCEINLINE float FInterpCurve<T>::InaccurateFindNearestOnSegment( const T &Poin
 }
 
 template< class T > 
-FORCEINLINE void FInterpCurve<T>::AutoSetTangents(float Tension)
+void FInterpCurve<T>::AutoSetTangents(float Tension)
 {
 	// Iterate over all points in this InterpCurve
 	for(int32 PointIndex=0; PointIndex<Points.Num(); PointIndex++)
@@ -540,7 +494,7 @@ FORCEINLINE void FInterpCurve<T>::AutoSetTangents(float Tension)
 }
 
 template< class T > 
-FORCEINLINE void FInterpCurve<T>::CalcBounds(T& OutMin, T& OutMax, const T& Default) const
+void FInterpCurve<T>::CalcBounds(T& OutMin, T& OutMax, const T& Default) const
 {
 	if(Points.Num() == 0)
 	{
diff --git a/Engine/Source/Runtime/Engine/Classes/Components/SplineComponent.h b/Engine/Source/Runtime/Engine/Classes/Components/SplineComponent.h
--- a/Engine/Source/Runtime/Engine/Classes/Components/SplineComponent.h
+++ b/Engine/Source/Runtime/Engine/Classes/Components/SplineComponent.h
@@ -19,6 +19,10 @@ class ENGINE_API USplineComponent : public USceneComponent
 	UPROPERTY()
 	FInterpCurveFloat SplineReparamTable;
 
+	/** Specifies the duration of the spline in seconds */
+	UPROPERTY(EditAnywhere, BlueprintReadWrite, Category=Spline)
+	float Duration;
+
 	// Begin UActorComponent interface.
 	virtual TSharedPtr<FComponentInstanceDataBase> GetComponentInstanceData() const OVERRIDE;
 	virtual FName GetComponentInstanceDataType() const OVERRIDE;
@@ -28,20 +32,56 @@ class ENGINE_API USplineComponent : public USceneComponent
 	/** Update the SplineReparamTable */
 	void UpdateSplineReparamTable();
 
+	/** Clears all the points in the spline */
+	UFUNCTION(BlueprintCallable, Category=Spline)
+	void ClearSplinePoints();
+
+	/** Adds a world space point to the spline */
+	UFUNCTION(BlueprintCallable, Category=Spline)
+	void AddSplineWorldPoint(const FVector& Position);
+
+	/** Sets the spline to an array of world space points */
+	UFUNCTION(BlueprintCallable, Category=Spline)
+	void SetSplineWorldPoints(const TArray<FVector>& Points);
+
 	/** Returns total length along this spline */
 	UFUNCTION(BlueprintCallable, Category=Spline) 
 	float GetSplineLength() const;
-	
+
+	/** Given a distance along the length of this spline, return the corresponding input key at that point */
+	UFUNCTION(BlueprintCallable, Category=Spline)
+	float GetInputKeyAtDistanceAlongSpline(float Distance) const;
+
 	/** Given a distance along the length of this spline, return the point in space where this puts you */
 	UFUNCTION(BlueprintCallable, Category=Spline)
 	FVector GetWorldLocationAtDistanceAlongSpline(float Distance) const;
 	
-	/** Given a distance along the length of this spline, return a unit direction vector of the spline there. */
+	/** Given a distance along the length of this spline, return a unit direction vector of the spline tangent there. */
 	UFUNCTION(BlueprintCallable, Category=Spline)
 	FVector GetWorldDirectionAtDistanceAlongSpline(float Distance) const;
 
+	/** Given a distance along the length of this spline, return a rotation corresponding to the spline's position and direction there. */
+	UFUNCTION(BlueprintCallable, Category = Spline)
+	FRotator GetWorldRotationAtDistanceAlongSpline(float Distance) const;
+
+	/** Given a time from 0 to the spline duration, return the point in space where this puts you */
+	UFUNCTION(BlueprintCallable, Category=Spline)
+	FVector GetWorldLocationAtTime(float Time, bool bUseConstantVelocity = false) const;
+
+	/** Given a time from 0 to the spline duration, return a unit direction vector of the spline tangent there. */
+	UFUNCTION(BlueprintCallable, Category=Spline)
+	FVector GetWorldDirectionAtTime(float Time, bool bUseConstantVelocity = false) const;
+
+	/** Given a time from 0 to the spline duration, return a rotation corresponding to the spline's position and direction there. */
+	UFUNCTION(BlueprintCallable, Category=Spline)
+	FRotator GetWorldRotationAtTime(float Time, bool bUseConstantVelocity = false) const;
+
 	/** Walk through keys and set time for each one */
 	void RefreshSplineInputs();
+
+private:
+	float GetSegmentLength(const int32 Index, const float Param = 1.0f) const;
+	float GetSegmentParamFromLength(const int32 Index, const float Length, const float SegmentLength) const;
 };
 
 
diff --git a/Engine/Source/Runtime/Engine/Private/SplineComponent.cpp b/Engine/Source/Runtime/Engine/Private/SplineComponent.cpp
--- a/Engine/Source/Runtime/Engine/Private/SplineComponent.cpp
+++ b/Engine/Source/Runtime/Engine/Private/SplineComponent.cpp
@@ -12,6 +12,8 @@
 USplineComponent::USplineComponent(const class FPostConstructInitializeProperties& PCIP)
 	: Super(PCIP)
 {
+	Duration = 1.0f;
+
 	// Add 2 keys by default
 	int32 PointIndex = SplineInfo.AddPoint(0.f, FVector(0,0,0));
 	SplineInfo.Points[PointIndex].InterpMode = CIM_CurveAuto;
@@ -33,49 +35,210 @@ void USplineComponent::UpdateSplineReparamTable()
 	{
 		return;
 	}
-	
-	const int32 NumKeys = SplineInfo.Points.Num(); // Number of keys on spline
-	const int32 NumSteps = 10 * (NumKeys-1); // TODO: Make this adaptive...
-	
-	// Find range of input
-	float Param = SplineInfo.Points[0].InVal;
-	const float MaxInput = SplineInfo.Points[NumKeys-1].InVal;
-	const float Interval = (MaxInput - Param)/((float)(NumSteps-1)); 
-	
-	// Add first entry, using first point on curve, total distance will be 0
-	FVector OldSplinePos = SplineInfo.Eval(Param, FVector::ZeroVector);
-	float TotalDist = 0.f;
-	SplineReparamTable.AddPoint(TotalDist, Param);
-	Param += Interval;
 
-	// Then work over rest of points	
-	for(int32 i=1; i<NumSteps; i++)
+	const int32 NumSegments = SplineInfo.Points.Num() - 1;
+	const int32 NumSteps = 4 * NumSegments;
+
+	// Build array of segment lengths
+	TArray<float> SegmentLengths;
+	SegmentLengths.Reset(NumSegments);
+
+	float TotalLength = 0.0f;
+	for (int32 SegmentIndex = 0; SegmentIndex < NumSegments; ++SegmentIndex)
+	{
+		float SegmentLength = GetSegmentLength(SegmentIndex);
+		SegmentLengths.Add(SegmentLength);
+		TotalLength += SegmentLength;
+	}
+
+	// Now step through the spine at a constant distance interval, adding points to the reparam table
+	int32 CurrentSegment = 0;
+	float AccumulatedLength = 0.0f;
+	for (int32 Step = 0; Step < NumSteps; ++Step)
+	{
+		const float TargetLength = TotalLength * Step / NumSteps;
+
+		while (TargetLength - AccumulatedLength > SegmentLengths[CurrentSegment])
+		{
+			AccumulatedLength += SegmentLengths[CurrentSegment];
+			CurrentSegment++;
+		}
+
+		const float Param = GetSegmentParamFromLength(CurrentSegment, TargetLength - AccumulatedLength, SegmentLengths[CurrentSegment]);
+		const int32 Index = SplineReparamTable.AddPoint(TargetLength, CurrentSegment + Param);
+		SplineReparamTable.Points[Index].InterpMode = CIM_CurveAuto;
+	}
+
+	const int32 Index = SplineReparamTable.AddPoint(TotalLength, NumSegments);
+	SplineReparamTable.Points[Index].InterpMode = CIM_CurveAuto;
+
+	SplineReparamTable.AutoSetTangents();
+}
+
+
+float USplineComponent::GetSegmentLength(const int32 Index, const float Param) const
+{
+	check(Index < SplineInfo.Points.Num() - 1);
+	check(Param >= 0.0f && Param <= 1.0f);
+
+	// Evaluate the length of a Hermite spline segment.
+	// This calculates the integral of |dP/dt| dt, where P(t) is the spline equation with components (x(t), y(t), z(t)).
+	// This isn't solvable analytically, so we use a numerical method (Legendre-Gauss quadrature) which performs very well
+	// with functions of this type, even with very few samples.  In this case, just 5 samples is sufficient to yield a
+	// reasonable result.
+
+	struct FLegendreGaussCoefficient
+	{
+		float Abscissa;
+		float Weight;
+	};
+
+	static const FLegendreGaussCoefficient LegendreGaussCoefficients[] =
+	{
+		{ 0.0f, 0.5688889f },
+		{ -0.5384693f, 0.47862867f },
+		{ 0.5384693f, 0.47862867f },
+		{ -0.90617985f, 0.23692688f },
+		{ 0.90617985f, 0.23692688f }
+	};
+
+	const auto& StartPoint = SplineInfo.Points[Index];
+	const auto& EndPoint = SplineInfo.Points[Index + 1];
+	check(EndPoint.InVal - StartPoint.InVal == 1.0f);
+
+	const auto& P0 = StartPoint.OutVal;
+	const auto& T0 = StartPoint.LeaveTangent;
+	const auto& P1 = EndPoint.OutVal;
+	const auto& T1 = EndPoint.ArriveTangent;
+
+	// Cache the coefficients to be fed into the function to calculate the spline derivative at each sample point as they are constant.
+	const FVector Coeff1 = ((P0 - P1) * 2.0f + T0 + T1) * 3.0f;
+	const FVector Coeff2 = (P1 - P0) * 6.0f - T0 * 4.0f - T1 * 2.0f;
+	const FVector Coeff3 = T0;
+
+	const float HalfParam = Param * 0.5f;
+
+	float Length = 0.0f;
+	for (const auto& LegendreGaussCoefficient : LegendreGaussCoefficients)
+	{
+		// Calculate derivative at each Legendre-Gauss sample, and perform a weighted sum
+		const float Alpha = HalfParam * (1.0f + LegendreGaussCoefficient.Abscissa);
+		const FVector Derivative = (Coeff1 * Alpha + Coeff2) * Alpha + Coeff3;
+		Length += Derivative.Size() * LegendreGaussCoefficient.Weight;
+	}
+	Length *= HalfParam;
+
+	return Length;
+}
+
+
+float USplineComponent::GetSegmentParamFromLength(const int32 Index, const float Length, const float SegmentLength) const
+{
+	if (SegmentLength == 0.0f)
 	{
-		// Iterate along spline at regular param intervals
-		const FVector NewSplinePos = SplineInfo.Eval(Param, FVector::ZeroVector);
-		
-		TotalDist += (NewSplinePos - OldSplinePos).Size();
-		OldSplinePos = NewSplinePos;
+		return 0.0f;
+	}
+
+	// Given a function P(x) which yields points along a spline with x = 0...1, we can define a function L(t) to be the
+	// Euclidian length of the spline from P(0) to P(t):
+	//
+	//    L(t) = integral of |dP/dt| dt
+	//         = integral of sqrt((dx/dt)^2 + (dy/dt)^2 + (dz/dt)^2) dt
+	//
+	// This method evaluates the inverse of this function, i.e. given a length d, it obtains a suitable value for t such that:
+	//    L(t) - d = 0
+	//
+	// We use Newton-Raphson to iteratively converge on the result:
+	//
+	//    t' = t - f(t) / (df/dt)
+	//
+	// where: t is an initial estimate of the result, obtained through basic linear interpolation,
+	//        f(t) is the function whose root we wish to find = L(t) - d,
+	//        (df/dt) = d(L(t))/dt = |dP/dt|
+
+	check(Index < SplineInfo.Points.Num() - 1);
+	check(Length >= 0.0f && Length <= SegmentLength);
+	check(SplineInfo.Points[Index + 1].InVal - SplineInfo.Points[Index].InVal == 1.0f);
 
-		SplineReparamTable.AddPoint(TotalDist, Param);
+	float Param = Length / SegmentLength;  // initial estimate for t
 
-		// move along
-		Param += Interval;
+	// two iterations of Newton-Raphson is enough
+	for (int32 Iteration = 0; Iteration < 2; ++Iteration)
+	{
+		float TangentMagnitude = SplineInfo.EvalDerivative(Index + Param, FVector::ZeroVector).Size();
+		if (TangentMagnitude > 0.0f)
+		{
+			Param -= (GetSegmentLength(Index, Param) - Length) / TangentMagnitude;
+			Param = FMath::Clamp(Param, 0.0f, 1.0f);
+		}
 	}
+
+	return Param;
+}
+
+
+void USplineComponent::ClearSplinePoints()
+{
+	SplineInfo.Reset();
+	SplineReparamTable.Reset();
+}
+
+
+void USplineComponent::AddSplineWorldPoint(const FVector& Position)
+{
+	float InputKey = static_cast<float>(SplineInfo.Points.Num());
+	int32 PointIndex = SplineInfo.AddPoint(InputKey, ComponentToWorld.InverseTransformPosition(Position));
+	SplineInfo.Points[PointIndex].InterpMode = CIM_CurveAuto;
+
+	SplineInfo.AutoSetTangents();
+	UpdateSplineReparamTable();
 }
 
+
+void USplineComponent::SetSplineWorldPoints(const TArray<FVector>& Points)
+{
+	SplineInfo.Reset();
+	float InputKey = 0.0f;
+	for (const auto& Point : Points)
+	{
+		int32 PointIndex = SplineInfo.AddPoint(InputKey, ComponentToWorld.InverseTransformPosition(Point));
+		SplineInfo.Points[PointIndex].InterpMode = CIM_CurveAuto;
+		InputKey += 1.0f;
+	}
+
+	SplineInfo.AutoSetTangents();
+	UpdateSplineReparamTable();
+}
+
+
 float USplineComponent::GetSplineLength() const
 {
+	const int32 NumPoints = SplineReparamTable.Points.Num();
+
 	// This is given by the input of the last entry in the remap table
-	if(SplineReparamTable.Points.Num() > 0)
+	if (NumPoints > 0)
 	{
-		return SplineReparamTable.Points[SplineReparamTable.Points.Num()-1].InVal;
+		return SplineReparamTable.Points[NumPoints - 1].InVal;
 	}
 	
 	return 0.f;
 }
 
 
+float USplineComponent::GetInputKeyAtDistanceAlongSpline(float Distance) const
+{
+	const int32 NumPoints = SplineInfo.Points.Num();
+
+	if (NumPoints < 2)
+	{
+		return 0.0f;
+	}
+
+	const float TimeMultiplier = Duration / (NumPoints - 1.0f);
+	return SplineReparamTable.Eval(Distance, 0.f) * TimeMultiplier;
+}
+
+
 FVector USplineComponent::GetWorldLocationAtDistanceAlongSpline(float Distance) const
 {
 	const float Param = SplineReparamTable.Eval(Distance, 0.f);
@@ -91,6 +254,69 @@ FVector USplineComponent::GetWorldDirectionAtDistanceAlongSpline(float Distance)
 	return ComponentToWorld.TransformVectorNoScale(Tangent);
 }
 
+
+FRotator USplineComponent::GetWorldRotationAtDistanceAlongSpline(float Distance) const
+{
+	const float Param = SplineReparamTable.Eval(Distance, 0.f);
+	const FVector Forward = SplineInfo.EvalDerivative(Param, FVector(1.0f, 0.0f, 0.0f));
+	return Forward.Rotation() + ComponentToWorld.Rotator();
+}
+
+
+FVector USplineComponent::GetWorldLocationAtTime(float Time, bool bUseConstantVelocity) const
+{
+	if (Duration == 0.0f)
+	{
+		return FVector::ZeroVector;
+	}
+
+	if (bUseConstantVelocity)
+	{
+		return GetWorldLocationAtDistanceAlongSpline(Time / Duration * GetSplineLength());
+	}
+
+	const float TimeMultiplier = (SplineInfo.Points.Num() - 1.0f) / Duration;
+	const FVector Location = SplineInfo.Eval(Time * TimeMultiplier, FVector::ZeroVector);
+	return ComponentToWorld.TransformPosition(Location);
+}
+
+
+FVector USplineComponent::GetWorldDirectionAtTime(float Time, bool bUseConstantVelocity) const
+{
+	if (Duration == 0.0f)
+	{
+		return FVector::ZeroVector;
+	}
+
+	if (bUseConstantVelocity)
+	{
+		return GetWorldDirectionAtDistanceAlongSpline(Time / Duration * GetSplineLength());
+	}
+
+	const float TimeMultiplier = (SplineInfo.Points.Num() - 1.0f) / Duration;
+	const FVector Tangent = SplineInfo.EvalDerivative(Time * TimeMultiplier, FVector::ZeroVector).SafeNormal();
+	return ComponentToWorld.TransformVectorNoScale(Tangent);
+}
+
+
+FRotator USplineComponent::GetWorldRotationAtTime(float Time, bool bUseConstantVelocity) const
+{
+	if (Duration == 0.0f)
+	{
+		return FRotator::ZeroRotator;
+	}
+
+	if (bUseConstantVelocity)
+	{
+		return GetWorldRotationAtDistanceAlongSpline(Time / Duration * GetSplineLength());
+	}
+
+	const float TimeMultiplier = (SplineInfo.Points.Num() - 1.0f) / Duration;
+	const FVector Forward = SplineInfo.EvalDerivative(Time * TimeMultiplier, FVector(1.0f, 0.0f, 0.0f));
+	return Forward.Rotation() + ComponentToWorld.Rotator();
+}
+
+
 void USplineComponent::RefreshSplineInputs()
 {
 	for(int32 KeyIdx=0; KeyIdx<SplineInfo.Points.Num(); KeyIdx++)
