Editor: Splines now trigger a PostEditChangeProperty on the component when modified via the component visualizer

diff --git a/Engine/Source/Editor/ComponentVisualizers/Private/SplineComponentVisualizer.cpp b/Engine/Source/Editor/ComponentVisualizers/Private/SplineComponentVisualizer.cpp
--- a/Engine/Source/Editor/ComponentVisualizers/Private/SplineComponentVisualizer.cpp
+++ b/Engine/Source/Editor/ComponentVisualizers/Private/SplineComponentVisualizer.cpp
@@ -82,7 +82,9 @@ void FSplineComponentVisualizer::DrawVisualization(const UActorComponent* Compon
 {
 	if (const USplineComponent* SplineComp = Cast<const USplineComponent>(Component))
 	{
-		FColor DrawColor(255,255,255);
+		const FColor DrawColor(255, 255, 255);
+		const FColor SelectedColor(255, 0, 0);
+		const float GrabHandleSize = 12.0f;
 
 		const FInterpCurveVector& SplineInfo = SplineComp->SplineInfo;
 
@@ -105,17 +107,17 @@ void FSplineComponentVisualizer::DrawVisualization(const UActorComponent* Compon
 			FVector NewKeyPos = SplineComp->ComponentToWorld.TransformPosition( SplineInfo.Eval(NewKeyTime, FVector::ZeroVector) );
 
 			USplineComponent* EditedSplineComp = GetEditedSplineComponent();
-			FColor KeyColor = (SplineComp == EditedSplineComp && KeyIdx == SelectedKeyIndex) ? FColor(255,0,0) : DrawColor;
+			const FColor KeyColor = (SplineComp == EditedSplineComp && KeyIdx == SelectedKeyIndex) ? SelectedColor : DrawColor;
 
 			// Draw the keypoint
 			PDI->SetHitProxy(new HSplineKeyProxy(Component, KeyIdx));
 
-			PDI->DrawPoint(NewKeyPos, KeyColor, 8.f, SDPG_World);
+			PDI->DrawPoint(NewKeyPos, KeyColor, GrabHandleSize, SDPG_World);
 
 			PDI->SetHitProxy(NULL);
 
 
-			// If not the first keypoint, draw a line to the last keypoint.
+			// If not the first keypoint, draw a line to the previous keypoint.
 			if (KeyIdx>0)
 			{
 				// For constant interpolation - don't draw ticks - just draw dotted line.
@@ -272,15 +274,7 @@ bool FSplineComponentVisualizer::HandleInputDelta(FLevelEditorViewportClient* Vi
 			EditedPoint.ArriveTangent = NewTangent;
 		}
 
-		// Update tangents and reparam table
-		SplineComp->SplineInfo.AutoSetTangents();
-		SplineComp->UpdateSplineReparamTable();
-
-		// Notify of change so any CS is re-run
-		if (SplineOwningActor.IsValid())
-		{
-			SplineOwningActor.Get()->PostEditMove(true);
-		}
+		NotifyComponentModified();
 
 		return true;
 	}
@@ -333,10 +327,8 @@ void FSplineComponentVisualizer::OnDuplicateKey()
 
 	// Update Input value for all keys
 	SplineComp->RefreshSplineInputs();
-	SplineComp->SplineInfo.AutoSetTangents(); // update tangents
-	SplineComp->UpdateSplineReparamTable(); // update reparam table
 
-	GEditor->RedrawLevelEditingViewports(true);
+	NotifyComponentModified();
 }
 
 void FSplineComponentVisualizer::OnDeleteKey()
@@ -346,17 +338,9 @@ void FSplineComponentVisualizer::OnDeleteKey()
 	SplineComp->SplineInfo.Points.RemoveAt(SelectedKeyIndex);
 
 	SplineComp->RefreshSplineInputs(); // update input value for each key
-	SplineComp->SplineInfo.AutoSetTangents(); // update tangents
-	SplineComp->UpdateSplineReparamTable(); // update reparam table
 	SelectedKeyIndex = INDEX_NONE; // deselect any keys
 
-	// Notify of change so any CS is re-run
-	if (SplineOwningActor.IsValid())
-	{
-		SplineOwningActor.Get()->PostEditMove(true);
-	}
-
-	GEditor->RedrawLevelEditingViewports(true);
+	NotifyComponentModified();
 }
 
 bool FSplineComponentVisualizer::IsSelectionValid() const
@@ -372,16 +356,8 @@ void FSplineComponentVisualizer::OnResetToAutomaticTangent()
 {
 	USplineComponent* SplineComp = GetEditedSplineComponent();
 	SplineComp->SplineInfo.Points[SelectedKeyIndex].InterpMode = CIM_CurveAuto;
-	SplineComp->SplineInfo.AutoSetTangents();
-	SplineComp->UpdateSplineReparamTable();
-
-	// Notify of change so any CS is re-run
-	if (SplineOwningActor.IsValid())
-	{
-		SplineOwningActor.Get()->PostEditMove(true);
-	}
 
-	GEditor->RedrawLevelEditingViewports(true);
+	NotifyComponentModified();
 }
 
 bool FSplineComponentVisualizer::CanResetToAutomaticTangent() const
@@ -411,4 +387,26 @@ TSharedPtr<SWidget> FSplineComponentVisualizer::GenerateContextMenu() const
 	return MenuWidget;
 }
 
+void FSplineComponentVisualizer::NotifyComponentModified()
+{
+	USplineComponent* SplineComp = GetEditedSplineComponent();
+
+	// Update tangents and reparam table
+	SplineComp->SplineInfo.AutoSetTangents();
+	SplineComp->UpdateSplineReparamTable();
+
+	// Notify that the spline info has been modified
+	UProperty* SplineInfoProperty = FindField<UProperty>(USplineComponent::StaticClass(), GET_MEMBER_NAME_CHECKED(USplineComponent, SplineInfo));
+	FPropertyChangedEvent PropertyChangedEvent(SplineInfoProperty);
+	SplineComp->PostEditChangeProperty(PropertyChangedEvent);
+
+	// Notify of change so any CS is re-run
+	if (SplineOwningActor.IsValid())
+	{
+		SplineOwningActor.Get()->PostEditMove(true);
+	}
+
+	GEditor->RedrawLevelEditingViewports(true);
+}
+
 #undef LOCTEXT_NAMESPACE
diff --git a/Engine/Source/Editor/ComponentVisualizers/Private/SplineComponentVisualizer.h b/Engine/Source/Editor/ComponentVisualizers/Private/SplineComponentVisualizer.h
--- a/Engine/Source/Editor/ComponentVisualizers/Private/SplineComponentVisualizer.h
+++ b/Engine/Source/Editor/ComponentVisualizers/Private/SplineComponentVisualizer.h
@@ -52,6 +52,8 @@ public:
 
 private:
 
+	void NotifyComponentModified();
+
 	void OnDeleteKey();
 	void OnDuplicateKey();
 	bool IsSelectionValid() const;
