Editor: Make spline editing widgets draw in foreground mode

diff --git a/Engine/Source/Editor/ComponentVisualizers/Private/SplineComponentVisualizer.cpp b/Engine/Source/Editor/ComponentVisualizers/Private/SplineComponentVisualizer.cpp
--- a/Engine/Source/Editor/ComponentVisualizers/Private/SplineComponentVisualizer.cpp
+++ b/Engine/Source/Editor/ComponentVisualizers/Private/SplineComponentVisualizer.cpp
@@ -112,7 +112,7 @@ void FSplineComponentVisualizer::DrawVisualization(const UActorComponent* Compon
 			// Draw the keypoint
 			PDI->SetHitProxy(new HSplineKeyProxy(Component, KeyIdx));
 
-			PDI->DrawPoint(NewKeyPos, KeyColor, GrabHandleSize, SDPG_World);
+			PDI->DrawPoint(NewKeyPos, KeyColor, GrabHandleSize, SDPG_Foreground);
 
 			PDI->SetHitProxy(NULL);
 
@@ -140,12 +140,12 @@ void FSplineComponentVisualizer::DrawVisualization(const UActorComponent* Compon
 						float NewTime = OldKeyTime + StepIdx*DrawSubstep;
 						FVector NewPos = SplineComp->ComponentToWorld.TransformPosition( SplineInfo.Eval(NewTime, FVector::ZeroVector) );
 
-						PDI->DrawLine(OldPos, NewPos, DrawColor, SDPG_World);
+						PDI->DrawLine(OldPos, NewPos, DrawColor, SDPG_Foreground);
 
 						// Don't draw point for last one - its the keypoint drawn above.
 						if (false)// (StepIdx != NumSteps)
 						{
-							PDI->DrawPoint(NewPos, DrawColor, 3.f, SDPG_World);
+							PDI->DrawPoint(NewPos, DrawColor, 3.f, SDPG_Foreground);
 						}
 
 						OldTime = NewTime;
