Further SplineComponent improvements:
- Changed USplineComponent::SplineReparamTable to a transient member and force rebuild each time in PostLoad() / PostDuplicate().
- Initialize the SplineReparamTable size at the start to save needless reallocations (as its size is known from the beginning).
- Added menu options to set spline key types (curve, linear, constant)
- Added Reset to Automatic Tangent to allow resetting to either a clamped or an unclamped tangent.
#branch UE4
#proj Runtime.Engine, Editor.ComponentVisualizers

diff --git a/Engine/Source/Editor/ComponentVisualizers/Private/SplineComponentVisualizer.cpp b/Engine/Source/Editor/ComponentVisualizers/Private/SplineComponentVisualizer.cpp
--- a/Engine/Source/Editor/ComponentVisualizers/Private/SplineComponentVisualizer.cpp
+++ b/Engine/Source/Editor/ComponentVisualizers/Private/SplineComponentVisualizer.cpp
@@ -2,6 +2,7 @@
 
 #include "ComponentVisualizersPrivatePCH.h"
 #include "SplineComponentVisualizer.h"
+#include "ScopedTransaction.h"
 
 IMPLEMENT_HIT_PROXY(HSplineVisProxy, HComponentVisProxy);
 IMPLEMENT_HIT_PROXY(HSplineKeyProxy, HSplineVisProxy);
@@ -25,9 +26,13 @@ public:
 
 	virtual void RegisterCommands() override
 	{
-		UI_COMMAND(DeleteKey, "Delete key", "Delete the currently selected key.", EUserInterfaceActionType::Button, FInputGesture(EKeys::Delete));
-		UI_COMMAND(DuplicateKey, "Duplicate key", "Duplicates the currently selected key.", EUserInterfaceActionType::Button, FInputGesture());
-		UI_COMMAND(ResetToAutomaticTangent, "Reset to Automatic Tangent", "Specify that the tangent for this key will be reset to its default auto-generated value.", EUserInterfaceActionType::Button, FInputGesture());
+		UI_COMMAND(DeleteKey, "Delete Key", "Delete the currently selected key.", EUserInterfaceActionType::Button, FInputGesture(EKeys::Delete));
+		UI_COMMAND(DuplicateKey, "Duplicate Key", "Duplicate the currently selected key.", EUserInterfaceActionType::Button, FInputGesture());
+		UI_COMMAND(ResetToUnclampedTangent, "Unclamped Tangent", "Reset the tangent for this key to its default unclamped value.", EUserInterfaceActionType::Button, FInputGesture());
+		UI_COMMAND(ResetToClampedTangent, "Clamped Tangent", "Reset the tangent for this key to its default clamped value.", EUserInterfaceActionType::Button, FInputGesture());
+		UI_COMMAND(SetKeyToCurve, "Curve", "Set spline key to Curve type", EUserInterfaceActionType::RadioButton, FInputGesture());
+		UI_COMMAND(SetKeyToLinear, "Linear", "Set spline key to Linear type", EUserInterfaceActionType::RadioButton, FInputGesture());
+		UI_COMMAND(SetKeyToConstant, "Constant", "Set spline key to Constant type", EUserInterfaceActionType::RadioButton, FInputGesture());
 	}
 
 public:
@@ -37,8 +42,20 @@ public:
 	/** Duplicate key */
 	TSharedPtr<FUICommandInfo> DuplicateKey;
 
-	/** Reset to automatic tangent */
-	TSharedPtr<FUICommandInfo> ResetToAutomaticTangent;
+	/** Reset to unclamped tangent */
+	TSharedPtr<FUICommandInfo> ResetToUnclampedTangent;
+
+	/** Reset to clamped tangent */
+	TSharedPtr<FUICommandInfo> ResetToClampedTangent;
+
+	/** Set spline key to Curve type */
+	TSharedPtr<FUICommandInfo> SetKeyToCurve;
+
+	/** Set spline key to Linear type */
+	TSharedPtr<FUICommandInfo> SetKeyToLinear;
+
+	/** Set spline key to Constant type */
+	TSharedPtr<FUICommandInfo> SetKeyToConstant;
 };
 
 
@@ -68,9 +85,32 @@ void FSplineComponentVisualizer::OnRegister()
 		FCanExecuteAction::CreateSP(this, &FSplineComponentVisualizer::IsSelectionValid));
 
 	SplineComponentVisualizerActions->MapAction(
-		Commands.ResetToAutomaticTangent,
-		FExecuteAction::CreateSP(this, &FSplineComponentVisualizer::OnResetToAutomaticTangent),
-		FCanExecuteAction::CreateSP(this, &FSplineComponentVisualizer::CanResetToAutomaticTangent));
+		Commands.ResetToUnclampedTangent,
+		FExecuteAction::CreateSP(this, &FSplineComponentVisualizer::OnResetToAutomaticTangent, CIM_CurveAuto),
+		FCanExecuteAction::CreateSP(this, &FSplineComponentVisualizer::CanResetToAutomaticTangent, CIM_CurveAuto));
+
+	SplineComponentVisualizerActions->MapAction(
+		Commands.ResetToClampedTangent,
+		FExecuteAction::CreateSP(this, &FSplineComponentVisualizer::OnResetToAutomaticTangent, CIM_CurveAutoClamped),
+		FCanExecuteAction::CreateSP(this, &FSplineComponentVisualizer::CanResetToAutomaticTangent, CIM_CurveAutoClamped));
+
+	SplineComponentVisualizerActions->MapAction(
+		Commands.SetKeyToCurve,
+		FExecuteAction::CreateSP(this, &FSplineComponentVisualizer::OnSetKeyType, CIM_CurveAuto),
+		FCanExecuteAction(),
+		FIsActionChecked::CreateSP(this, &FSplineComponentVisualizer::IsKeyTypeSet, CIM_CurveAuto));
+
+	SplineComponentVisualizerActions->MapAction(
+		Commands.SetKeyToLinear,
+		FExecuteAction::CreateSP(this, &FSplineComponentVisualizer::OnSetKeyType, CIM_Linear),
+		FCanExecuteAction(),
+		FIsActionChecked::CreateSP(this, &FSplineComponentVisualizer::IsKeyTypeSet, CIM_Linear));
+
+	SplineComponentVisualizerActions->MapAction(
+		Commands.SetKeyToConstant,
+		FExecuteAction::CreateSP(this, &FSplineComponentVisualizer::OnSetKeyType, CIM_Constant),
+		FCanExecuteAction(),
+		FIsActionChecked::CreateSP(this, &FSplineComponentVisualizer::IsKeyTypeSet, CIM_Constant));
 }
 
 FSplineComponentVisualizer::~FSplineComponentVisualizer()
@@ -82,6 +122,11 @@ void FSplineComponentVisualizer::DrawVisualization(const UActorComponent* Compon
 {
 	if (const USplineComponent* SplineComp = Cast<const USplineComponent>(Component))
 	{
+		if (!SplineComp->bAllowSplineEditingPerInstance)
+		{
+			return;
+		}
+
 		const FColor DrawColor(255, 255, 255);
 		const FColor SelectedColor(255, 0, 0);
 		const float GrabHandleSize = 12.0f;
@@ -111,12 +156,9 @@ void FSplineComponentVisualizer::DrawVisualization(const UActorComponent* Compon
 
 			// Draw the keypoint
 			PDI->SetHitProxy(new HSplineKeyProxy(Component, KeyIdx));
-
 			PDI->DrawPoint(NewKeyPos, KeyColor, GrabHandleSize, SDPG_Foreground);
-
 			PDI->SetHitProxy(NULL);
 
-
 			// If not the first keypoint, draw a line to the previous keypoint.
 			if (KeyIdx>0)
 			{
@@ -142,12 +184,6 @@ void FSplineComponentVisualizer::DrawVisualization(const UActorComponent* Compon
 
 						PDI->DrawLine(OldPos, NewPos, DrawColor, SDPG_Foreground);
 
-						// Don't draw point for last one - its the keypoint drawn above.
-						if (false)// (StepIdx != NumSteps)
-						{
-							PDI->DrawPoint(NewPos, DrawColor, 3.f, SDPG_Foreground);
-						}
-
 						OldTime = NewTime;
 						OldPos = NewPos;
 					}
@@ -311,8 +347,14 @@ void FSplineComponentVisualizer::EndEditing()
 
 void FSplineComponentVisualizer::OnDuplicateKey()
 {
+	const FScopedTransaction Transaction(LOCTEXT("DuplicateKey", "Duplicate Spline Key"));
 	USplineComponent* SplineComp = GetEditedSplineComponent();
 
+	if (AActor* Owner = SplineComp->GetOwner())
+	{
+		Owner->Modify();
+	}
+
 	FInterpCurvePoint<FVector> KeyToDupe = SplineComp->SplineInfo.Points[SelectedKeyIndex];
 	KeyToDupe.InterpMode = CIM_CurveAuto;
 	int32 NewKeyIndex = SplineComp->SplineInfo.Points.Insert(KeyToDupe, SelectedKeyIndex);
@@ -325,18 +367,26 @@ void FSplineComponentVisualizer::OnDuplicateKey()
 	NotifyComponentModified();
 }
 
+
 void FSplineComponentVisualizer::OnDeleteKey()
 {
+	const FScopedTransaction Transaction(LOCTEXT("DeleteKey", "Delete Spline Key"));
 	USplineComponent* SplineComp = GetEditedSplineComponent();
 
-	SplineComp->SplineInfo.Points.RemoveAt(SelectedKeyIndex);
+	if (AActor* Owner = SplineComp->GetOwner())
+	{
+		Owner->Modify();
+	}
 
+	SplineComp->SplineInfo.Points.RemoveAt(SelectedKeyIndex);
 	SplineComp->RefreshSplineInputs(); // update input value for each key
+
 	SelectedKeyIndex = INDEX_NONE; // deselect any keys
 
 	NotifyComponentModified();
 }
 
+
 bool FSplineComponentVisualizer::IsSelectionValid() const
 {
 	USplineComponent* SplineComp = GetEditedSplineComponent();
@@ -346,21 +396,64 @@ bool FSplineComponentVisualizer::IsSelectionValid() const
 }
 
 
-void FSplineComponentVisualizer::OnResetToAutomaticTangent()
+void FSplineComponentVisualizer::OnResetToAutomaticTangent(EInterpCurveMode Mode)
 {
 	USplineComponent* SplineComp = GetEditedSplineComponent();
-	SplineComp->SplineInfo.Points[SelectedKeyIndex].InterpMode = CIM_CurveAuto;
 
-	NotifyComponentModified();
+	if (SplineComp->SplineInfo.Points[SelectedKeyIndex].InterpMode != Mode)
+	{
+		const FScopedTransaction Transaction(LOCTEXT("ResetToAutomaticTangent", "Reset to Automatic Tangent"));
+		if (AActor* Owner = SplineComp->GetOwner())
+		{
+			Owner->Modify();
+		}
+
+		SplineComp->SplineInfo.Points[SelectedKeyIndex].InterpMode = Mode;
+		NotifyComponentModified();
+	}
 }
 
-bool FSplineComponentVisualizer::CanResetToAutomaticTangent() const
+
+bool FSplineComponentVisualizer::CanResetToAutomaticTangent(EInterpCurveMode Mode) const
 {
 	USplineComponent* SplineComp = GetEditedSplineComponent();
 	return (SplineComp &&
 		SelectedKeyIndex != INDEX_NONE &&
 		SelectedKeyIndex < SplineComp->SplineInfo.Points.Num() &&
-		SplineComp->SplineInfo.Points[SelectedKeyIndex].InterpMode != CIM_CurveAuto);
+		SplineComp->SplineInfo.Points[SelectedKeyIndex].IsCurveKey() &&
+		SplineComp->SplineInfo.Points[SelectedKeyIndex].InterpMode != Mode);
+}
+
+
+void FSplineComponentVisualizer::OnSetKeyType(EInterpCurveMode Mode)
+{
+	USplineComponent* SplineComp = GetEditedSplineComponent();
+
+	if (SplineComp->SplineInfo.Points[SelectedKeyIndex].InterpMode != Mode)
+	{
+		const FScopedTransaction Transaction(LOCTEXT("ResetToAutomaticTangent", "Reset to Automatic Tangent"));
+		if (AActor* Owner = SplineComp->GetOwner())
+		{
+			Owner->Modify();
+		}
+
+		SplineComp->SplineInfo.Points[SelectedKeyIndex].InterpMode = Mode;
+		NotifyComponentModified();
+	}
+}
+
+
+bool FSplineComponentVisualizer::IsKeyTypeSet(EInterpCurveMode Mode) const
+{
+	if (IsSelectionValid())
+	{
+		USplineComponent* SplineComp = GetEditedSplineComponent();
+
+		const auto& SelectedPoint = SplineComp->SplineInfo.Points[SelectedKeyIndex];
+		return ((Mode == CIM_CurveAuto && SelectedPoint.IsCurveKey()) || SelectedPoint.InterpMode == Mode);
+	}
+
+	return false;
 }
 
 
@@ -372,7 +465,24 @@ TSharedPtr<SWidget> FSplineComponentVisualizer::GenerateContextMenu() const
 		{
 			MenuBuilder.AddMenuEntry(FSplineComponentVisualizerCommands::Get().DeleteKey);
 			MenuBuilder.AddMenuEntry(FSplineComponentVisualizerCommands::Get().DuplicateKey);
-			MenuBuilder.AddMenuEntry(FSplineComponentVisualizerCommands::Get().ResetToAutomaticTangent);
+
+			MenuBuilder.AddSubMenu(
+				LOCTEXT("KeyType", "Key Type"),
+				LOCTEXT("KeyTypeTooltip", "Define the type of the spline key."),
+				FNewMenuDelegate::CreateSP(this, &FSplineComponentVisualizer::GenerateSplinePointTypeSubMenu));
+
+			// Only add the Automatic Tangents submenu if the key is a curve type
+			USplineComponent* SplineComp = GetEditedSplineComponent();
+			if (SplineComp &&
+				SelectedKeyIndex != INDEX_NONE &&
+				SelectedKeyIndex < SplineComp->SplineInfo.Points.Num() &&
+				SplineComp->SplineInfo.Points[SelectedKeyIndex].IsCurveKey())
+			{
+				MenuBuilder.AddSubMenu(
+					LOCTEXT("ResetToAutomaticTangent", "Reset to Automatic Tangent"),
+					LOCTEXT("ResetToAutomaticTangentTooltip", "Reset the point tangent to an automatically generated value."),
+					FNewMenuDelegate::CreateSP(this, &FSplineComponentVisualizer::GenerateTangentTypeSubMenu));
+			}
 		}
 		MenuBuilder.EndSection();
 	}
@@ -381,6 +491,22 @@ TSharedPtr<SWidget> FSplineComponentVisualizer::GenerateContextMenu() const
 	return MenuWidget;
 }
 
+
+void FSplineComponentVisualizer::GenerateSplinePointTypeSubMenu(FMenuBuilder& MenuBuilder) const
+{
+	MenuBuilder.AddMenuEntry(FSplineComponentVisualizerCommands::Get().SetKeyToCurve);
+	MenuBuilder.AddMenuEntry(FSplineComponentVisualizerCommands::Get().SetKeyToLinear);
+	MenuBuilder.AddMenuEntry(FSplineComponentVisualizerCommands::Get().SetKeyToConstant);
+}
+
+
+void FSplineComponentVisualizer::GenerateTangentTypeSubMenu(FMenuBuilder& MenuBuilder) const
+{
+	MenuBuilder.AddMenuEntry(FSplineComponentVisualizerCommands::Get().ResetToUnclampedTangent);
+	MenuBuilder.AddMenuEntry(FSplineComponentVisualizerCommands::Get().ResetToClampedTangent);
+}
+
+
 void FSplineComponentVisualizer::NotifyComponentModified()
 {
 	USplineComponent* SplineComp = GetEditedSplineComponent();
diff --git a/Engine/Source/Editor/ComponentVisualizers/Private/SplineComponentVisualizer.h b/Engine/Source/Editor/ComponentVisualizers/Private/SplineComponentVisualizer.h
--- a/Engine/Source/Editor/ComponentVisualizers/Private/SplineComponentVisualizer.h
+++ b/Engine/Source/Editor/ComponentVisualizers/Private/SplineComponentVisualizer.h
@@ -58,8 +58,17 @@ private:
 	void OnDuplicateKey();
 	bool IsSelectionValid() const;
 
-	void OnResetToAutomaticTangent();
-	bool CanResetToAutomaticTangent() const;
+	void OnResetToAutomaticTangent(EInterpCurveMode Mode);
+	bool CanResetToAutomaticTangent(EInterpCurveMode Mode) const;
+
+	void OnSetKeyType(EInterpCurveMode Mode);
+	bool IsKeyTypeSet(EInterpCurveMode Mode) const;
+
+	/** Generate the submenu containing the available point types */
+	void GenerateSplinePointTypeSubMenu(FMenuBuilder& MenuBuilder) const;
+
+	/** Generate the submenu containing the available auto tangent types */
+	void GenerateTangentTypeSubMenu(FMenuBuilder& MenuBuilder) const;
 
 	/** Output log commands */
 	TSharedPtr<FUICommandList> SplineComponentVisualizerActions;
diff --git a/Engine/Source/Runtime/Engine/Classes/Components/SplineComponent.h b/Engine/Source/Runtime/Engine/Classes/Components/SplineComponent.h
--- a/Engine/Source/Runtime/Engine/Classes/Components/SplineComponent.h
+++ b/Engine/Source/Runtime/Engine/Classes/Components/SplineComponent.h
@@ -16,9 +16,13 @@ class ENGINE_API USplineComponent : public USceneComponent
 	FInterpCurveVector SplineInfo;
 
 	/** Input, distance along curve, output, parameter that puts you there. */
-	UPROPERTY()
+	UPROPERTY(Transient, TextExportTransient)
 	FInterpCurveFloat SplineReparamTable;
 
+	/** If true, spline keys may be edited per instance in the level viewport. Otherwise, the spline should be initialized in the construction script. */
+	UPROPERTY(EditDefaultsOnly, Category = Spline)
+	bool bAllowSplineEditingPerInstance;
+
 	/** Number of steps per spline segment to place in the reparameterization table */
 	UPROPERTY(EditAnywhere, AdvancedDisplay, Category = Spline, meta=(ClampMin=4, UIMin=4, ClampMax=100, UIMax=100))
 	int32 ReparamStepsPerSegment;
@@ -48,10 +52,18 @@ class ENGINE_API USplineComponent : public USceneComponent
 	UFUNCTION(BlueprintCallable, Category=Spline)
 	void AddSplineWorldPoint(const FVector& Position);
 
+	/** Adds a local space point to the spline */
+	UFUNCTION(BlueprintCallable, Category = Spline)
+	void AddSplineLocalPoint(const FVector& Position);
+
 	/** Sets the spline to an array of world space points */
 	UFUNCTION(BlueprintCallable, Category=Spline)
 	void SetSplineWorldPoints(const TArray<FVector>& Points);
 
+	/** Sets the spline to an array of local space points */
+	UFUNCTION(BlueprintCallable, Category = Spline)
+	void SetSplineLocalPoints(const TArray<FVector>& Points);
+
 	/** Move an existing point to a new world location */
 	UFUNCTION(BlueprintCallable, Category = Spline)
 	void SetWorldLocationAtSplinePoint(int32 PointIndex, const FVector& InLocation);
@@ -100,8 +112,10 @@ class ENGINE_API USplineComponent : public USceneComponent
 	UFUNCTION(BlueprintCallable, Category=Spline)
 	FRotator GetWorldRotationAtTime(float Time, bool bUseConstantVelocity = false) const;
 
-#if WITH_EDITORONLY_DATA
 	// UObject interface
+	virtual void PostLoad() override;
+	virtual void PostDuplicate(bool bDuplicateForPIE) override;
+#if WITH_EDITORONLY_DATA
 	virtual void PostEditChangeProperty(FPropertyChangedEvent& PropertyChangedEvent) override;
 	// End of UObject interface
 #endif
diff --git a/Engine/Source/Runtime/Engine/Private/SplineComponent.cpp b/Engine/Source/Runtime/Engine/Private/SplineComponent.cpp
--- a/Engine/Source/Runtime/Engine/Private/SplineComponent.cpp
+++ b/Engine/Source/Runtime/Engine/Private/SplineComponent.cpp
@@ -11,10 +11,13 @@
 
 USplineComponent::USplineComponent(const class FPostConstructInitializeProperties& PCIP)
 	: Super(PCIP)
+	, bAllowSplineEditingPerInstance(true)
 	, ReparamStepsPerSegment(10)
 	, Duration(1.0f)
 	, bStationaryEndpoints(false)
 {
+	SplineInfo.Points.Reset(10);
+
 	// Add 2 keys by default
 	int32 PointIndex = SplineInfo.AddPoint(0.f, FVector(0,0,0));
 	SplineInfo.Points[PointIndex].InterpMode = CIM_CurveAuto;
@@ -25,14 +28,26 @@ USplineComponent::USplineComponent(const class FPostConstructInitializePropertie
 	UpdateSpline();
 }
 
+
+void USplineComponent::PostLoad()
+{
+	Super::PostLoad();
+	UpdateSpline();
+}
+
+
+void USplineComponent::PostDuplicate(bool bDuplicateForPIE)
+{
+	Super::PostDuplicate(bDuplicateForPIE);
+	UpdateSpline();
+}
+
+
 void USplineComponent::UpdateSpline()
 {
 	// Automatically set the tangents on any CurveAuto keys
 	SplineInfo.AutoSetTangents(0.0f, bStationaryEndpoints);
 
-	// Start by clearing it
-	SplineReparamTable.Reset();
-	
 	// Nothing to do if no points
 	if(SplineInfo.Points.Num() < 2)
 	{
@@ -41,6 +56,8 @@ void USplineComponent::UpdateSpline()
 
 	const int32 NumSegments = SplineInfo.Points.Num() - 1;
 
+	// Start by clearing it
+	SplineReparamTable.Points.Reset(NumSegments * ReparamStepsPerSegment + 1);
 	float AccumulatedLength = 0.0f;
 	for (int32 SegmentIndex = 0; SegmentIndex < NumSegments; ++SegmentIndex)
 	{
@@ -84,7 +101,7 @@ float USplineComponent::GetSegmentLength(const int32 Index, const float Param) c
 
 	const auto& StartPoint = SplineInfo.Points[Index];
 	const auto& EndPoint = SplineInfo.Points[Index + 1];
-	check(EndPoint.InVal - StartPoint.InVal == 1.0f);
+	check(static_cast<int32>(EndPoint.InVal) - static_cast<int32>(StartPoint.InVal) == 1);
 
 	const auto& P0 = StartPoint.OutVal;
 	const auto& T0 = StartPoint.LeaveTangent;
@@ -138,7 +155,7 @@ float USplineComponent::GetSegmentParamFromLength(const int32 Index, const float
 
 	check(Index < SplineInfo.Points.Num() - 1);
 	check(Length >= 0.0f && Length <= SegmentLength);
-	check(SplineInfo.Points[Index + 1].InVal - SplineInfo.Points[Index].InVal == 1.0f);
+	check(static_cast<int32>(SplineInfo.Points[Index + 1].InVal) - static_cast<int32>(SplineInfo.Points[Index].InVal) == 1);
 
 	float Param = Length / SegmentLength;  // initial estimate for t
 
@@ -174,9 +191,19 @@ void USplineComponent::AddSplineWorldPoint(const FVector& Position)
 }
 
 
+void USplineComponent::AddSplineLocalPoint(const FVector& Position)
+{
+	float InputKey = static_cast<float>(SplineInfo.Points.Num());
+	int32 PointIndex = SplineInfo.AddPoint(InputKey, Position);
+	SplineInfo.Points[PointIndex].InterpMode = CIM_CurveAuto;
+
+	UpdateSpline();
+}
+
+
 void USplineComponent::SetSplineWorldPoints(const TArray<FVector>& Points)
 {
-	SplineInfo.Reset();
+	SplineInfo.Points.Reset(Points.Num());
 	float InputKey = 0.0f;
 	for (const auto& Point : Points)
 	{
@@ -188,6 +215,22 @@ void USplineComponent::SetSplineWorldPoints(const TArray<FVector>& Points)
 	UpdateSpline();
 }
 
+
+void USplineComponent::SetSplineLocalPoints(const TArray<FVector>& Points)
+{
+	SplineInfo.Points.Reset(Points.Num());
+	float InputKey = 0.0f;
+	for (const auto& Point : Points)
+	{
+		int32 PointIndex = SplineInfo.AddPoint(InputKey, Point);
+		SplineInfo.Points[PointIndex].InterpMode = CIM_CurveAuto;
+		InputKey += 1.0f;
+	}
+
+	UpdateSpline();
+}
+
+
 void USplineComponent::SetWorldLocationAtSplinePoint(int32 PointIndex, const FVector& InLocation)
 {
 	if((PointIndex >= 0) && (PointIndex < SplineInfo.Points.Num()))
@@ -342,13 +385,15 @@ void USplineComponent::RefreshSplineInputs()
 class FSplineInstanceData : public FComponentInstanceDataBase
 {
 public:
-	FSplineInstanceData(const USplineComponent* SourceComponent)
+	explicit FSplineInstanceData(const USplineComponent* SourceComponent)
 		: FComponentInstanceDataBase(SourceComponent)
-		, SplineInfo(SourceComponent->SplineInfo)
 	{
 	}
 
 	FInterpCurveVector SplineInfo;
+	int32 ReparamStepsPerSegment;
+	float Duration;
+	bool bStationaryEndpoints;
 };
 
 FName USplineComponent::GetComponentInstanceDataType() const
@@ -359,13 +404,32 @@ FName USplineComponent::GetComponentInstanceDataType() const
 
 TSharedPtr<FComponentInstanceDataBase> USplineComponent::GetComponentInstanceData() const
 {
-	return MakeShareable(new FSplineInstanceData(this));
+	TSharedPtr<FSplineInstanceData> SplineInstanceData;
+	if (bAllowSplineEditingPerInstance)
+	{
+		SplineInstanceData = MakeShareable(new FSplineInstanceData(this));
+		SplineInstanceData->SplineInfo = SplineInfo;
+		SplineInstanceData->ReparamStepsPerSegment = ReparamStepsPerSegment;
+		SplineInstanceData->Duration = Duration;
+		SplineInstanceData->bStationaryEndpoints = bStationaryEndpoints;
+	}
+	return SplineInstanceData;
 }
 
 void USplineComponent::ApplyComponentInstanceData(TSharedPtr<FComponentInstanceDataBase> ComponentInstanceData)
 {
-	SplineInfo = StaticCastSharedPtr<FSplineInstanceData>(ComponentInstanceData)->SplineInfo;
-	UpdateSpline();
+	if (ComponentInstanceData.IsValid())
+	{
+		const FSplineInstanceData* SplineInstanceData = StaticCastSharedPtr<FSplineInstanceData>(ComponentInstanceData).Get();
+		if (bAllowSplineEditingPerInstance)
+		{
+			SplineInfo = SplineInstanceData->SplineInfo;
+			ReparamStepsPerSegment = SplineInstanceData->ReparamStepsPerSegment;
+			Duration = SplineInstanceData->Duration;
+			bStationaryEndpoints = SplineInstanceData->bStationaryEndpoints;
+			UpdateSpline();
+		}
+	}
 }
 
 #if WITH_EDITORONLY_DATA
@@ -383,4 +447,4 @@ void USplineComponent::PostEditChangeProperty(FPropertyChangedEvent& PropertyCha
 
 	Super::PostEditChangeProperty(PropertyChangedEvent);
 }
-#endif
\ No newline at end of file
+#endif
