Miscellaneous improvements to SplineComponentVisualizer.
#ttp 341724 - FEATURE REQUEST: PAPER2D: Add or insert a key into a terrain spline
#ttp 341720 - FEATURE REQUEST: PAPER2D: Rename Key on terrain splines
#ttp 341378 - Spline Component/Meshs: Controls points only editable in worldspace. no localpace transformation

diff --git a/Engine/Source/Editor/ComponentVisualizers/Private/SplineComponentVisualizer.cpp b/Engine/Source/Editor/ComponentVisualizers/Private/SplineComponentVisualizer.cpp
--- a/Engine/Source/Editor/ComponentVisualizers/Private/SplineComponentVisualizer.cpp
+++ b/Engine/Source/Editor/ComponentVisualizers/Private/SplineComponentVisualizer.cpp
@@ -6,6 +6,7 @@
 
 IMPLEMENT_HIT_PROXY(HSplineVisProxy, HComponentVisProxy);
 IMPLEMENT_HIT_PROXY(HSplineKeyProxy, HSplineVisProxy);
+IMPLEMENT_HIT_PROXY(HSplineSegmentProxy, HSplineVisProxy);
 
 #define LOCTEXT_NAMESPACE "SplineComponentVisualizer"
 
@@ -26,13 +27,14 @@ public:
 
 	virtual void RegisterCommands() override
 	{
-		UI_COMMAND(DeleteKey, "Delete Key", "Delete the currently selected key.", EUserInterfaceActionType::Button, FInputGesture(EKeys::Delete));
-		UI_COMMAND(DuplicateKey, "Duplicate Key", "Duplicate the currently selected key.", EUserInterfaceActionType::Button, FInputGesture());
-		UI_COMMAND(ResetToUnclampedTangent, "Unclamped Tangent", "Reset the tangent for this key to its default unclamped value.", EUserInterfaceActionType::Button, FInputGesture());
-		UI_COMMAND(ResetToClampedTangent, "Clamped Tangent", "Reset the tangent for this key to its default clamped value.", EUserInterfaceActionType::Button, FInputGesture());
-		UI_COMMAND(SetKeyToCurve, "Curve", "Set spline key to Curve type", EUserInterfaceActionType::RadioButton, FInputGesture());
-		UI_COMMAND(SetKeyToLinear, "Linear", "Set spline key to Linear type", EUserInterfaceActionType::RadioButton, FInputGesture());
-		UI_COMMAND(SetKeyToConstant, "Constant", "Set spline key to Constant type", EUserInterfaceActionType::RadioButton, FInputGesture());
+		UI_COMMAND(DeleteKey, "Delete Spline Point", "Delete the currently selected spline point.", EUserInterfaceActionType::Button, FInputGesture(EKeys::Delete));
+		UI_COMMAND(DuplicateKey, "Duplicate Spline Point", "Duplicate the currently selected spline point.", EUserInterfaceActionType::Button, FInputGesture());
+		UI_COMMAND(AddKey, "Add Spline Point Here", "Add a new spline point at the cursor location.", EUserInterfaceActionType::Button, FInputGesture());
+		UI_COMMAND(ResetToUnclampedTangent, "Unclamped Tangent", "Reset the tangent for this spline point to its default unclamped value.", EUserInterfaceActionType::Button, FInputGesture());
+		UI_COMMAND(ResetToClampedTangent, "Clamped Tangent", "Reset the tangent for this spline point to its default clamped value.", EUserInterfaceActionType::Button, FInputGesture());
+		UI_COMMAND(SetKeyToCurve, "Curve", "Set spline point to Curve type", EUserInterfaceActionType::RadioButton, FInputGesture());
+		UI_COMMAND(SetKeyToLinear, "Linear", "Set spline point to Linear type", EUserInterfaceActionType::RadioButton, FInputGesture());
+		UI_COMMAND(SetKeyToConstant, "Constant", "Set spline point to Constant type", EUserInterfaceActionType::RadioButton, FInputGesture());
 	}
 
 public:
@@ -42,6 +44,9 @@ public:
 	/** Duplicate key */
 	TSharedPtr<FUICommandInfo> DuplicateKey;
 
+	/** Add key */
+	TSharedPtr<FUICommandInfo> AddKey;
+
 	/** Reset to unclamped tangent */
 	TSharedPtr<FUICommandInfo> ResetToUnclampedTangent;
 
@@ -63,6 +68,7 @@ public:
 FSplineComponentVisualizer::FSplineComponentVisualizer()
 	: FComponentVisualizer()
 	, SelectedKeyIndex(INDEX_NONE)
+	, SelectedSegmentIndex(INDEX_NONE)
 	, bAllowDuplication(true)
 {
 	FSplineComponentVisualizerCommands::Register();
@@ -77,12 +83,17 @@ void FSplineComponentVisualizer::OnRegister()
 	SplineComponentVisualizerActions->MapAction(
 		Commands.DeleteKey,
 		FExecuteAction::CreateSP(this, &FSplineComponentVisualizer::OnDeleteKey),
-		FCanExecuteAction::CreateSP(this, &FSplineComponentVisualizer::IsSelectionValid));
+		FCanExecuteAction::CreateSP(this, &FSplineComponentVisualizer::IsKeySelectionValid));
 
 	SplineComponentVisualizerActions->MapAction(
 		Commands.DuplicateKey,
 		FExecuteAction::CreateSP(this, &FSplineComponentVisualizer::OnDuplicateKey),
-		FCanExecuteAction::CreateSP(this, &FSplineComponentVisualizer::IsSelectionValid));
+		FCanExecuteAction::CreateSP(this, &FSplineComponentVisualizer::IsKeySelectionValid));
+
+	SplineComponentVisualizerActions->MapAction(
+		Commands.AddKey,
+		FExecuteAction::CreateSP(this, &FSplineComponentVisualizer::OnAddKey),
+		FCanExecuteAction::CreateSP(this, &FSplineComponentVisualizer::CanAddKey));
 
 	SplineComponentVisualizerActions->MapAction(
 		Commands.ResetToUnclampedTangent,
@@ -127,23 +138,14 @@ void FSplineComponentVisualizer::DrawVisualization(const UActorComponent* Compon
 			return;
 		}
 
-		const FColor DrawColor(255, 255, 255);
+		USplineComponent* EditedSplineComp = GetEditedSplineComponent();
+
+		const FColor NormalColor(255, 255, 255);
 		const FColor SelectedColor(255, 0, 0);
 		const float GrabHandleSize = 12.0f;
 
 		const FInterpCurveVector& SplineInfo = SplineComp->SplineInfo;
 
-		/*
-		// Debug draw keys on the SplineReparamTable
-		const FInterpCurveFloat& SplineReparamTable = SplineComp->SplineReparamTable;
-		for (const auto& Point : SplineReparamTable.Points)
-		{
-			const float Value = Point.OutVal;
-			const FVector KeyPos = SplineComp->ComponentToWorld.TransformPosition(SplineInfo.Eval(Value, FVector::ZeroVector));
-			PDI->DrawPoint(KeyPos, FColor(255, 0, 255), 10.f, SDPG_World);
-		}
-		*/
-
 		FVector OldKeyPos(0);
 		float OldKeyTime = 0.f;
 		for (int32 KeyIdx = 0; KeyIdx<SplineInfo.Points.Num(); KeyIdx++)
@@ -151,8 +153,7 @@ void FSplineComponentVisualizer::DrawVisualization(const UActorComponent* Compon
 			float NewKeyTime = SplineInfo.Points[KeyIdx].InVal;
 			FVector NewKeyPos = SplineComp->ComponentToWorld.TransformPosition( SplineInfo.Eval(NewKeyTime, FVector::ZeroVector) );
 
-			USplineComponent* EditedSplineComp = GetEditedSplineComponent();
-			const FColor KeyColor = (SplineComp == EditedSplineComp && KeyIdx == SelectedKeyIndex) ? SelectedColor : DrawColor;
+			const FColor KeyColor = (SplineComp == EditedSplineComp && KeyIdx == SelectedKeyIndex) ? SelectedColor : NormalColor;
 
 			// Draw the keypoint
 			PDI->SetHitProxy(new HSplineKeyProxy(Component, KeyIdx));
@@ -162,14 +163,17 @@ void FSplineComponentVisualizer::DrawVisualization(const UActorComponent* Compon
 			// If not the first keypoint, draw a line to the previous keypoint.
 			if (KeyIdx>0)
 			{
+				const FColor LineColor = (SplineComp == EditedSplineComp && KeyIdx == SelectedKeyIndex + 1) ? SelectedColor : NormalColor;
+				PDI->SetHitProxy(new HSplineSegmentProxy(Component, KeyIdx - 1));
+
 				// For constant interpolation - don't draw ticks - just draw dotted line.
 				if (SplineInfo.Points[KeyIdx - 1].InterpMode == CIM_Constant)
 				{
-					DrawDashedLine(PDI, OldKeyPos, NewKeyPos, DrawColor, 20, SDPG_World);
+					DrawDashedLine(PDI, OldKeyPos, NewKeyPos, LineColor, 20, SDPG_World);
 				}
 				else
 				{
-					int32 NumSteps = FMath::CeilToInt((NewKeyTime - OldKeyTime) / 0.02f);
+					int32 NumSteps = FMath::CeilToInt((NewKeyTime - OldKeyTime) * 32.0f);
 					float DrawSubstep = (NewKeyTime - OldKeyTime) / NumSteps;
 
 					// Find position on first keyframe.
@@ -182,12 +186,14 @@ void FSplineComponentVisualizer::DrawVisualization(const UActorComponent* Compon
 						float NewTime = OldKeyTime + StepIdx*DrawSubstep;
 						FVector NewPos = SplineComp->ComponentToWorld.TransformPosition( SplineInfo.Eval(NewTime, FVector::ZeroVector) );
 
-						PDI->DrawLine(OldPos, NewPos, DrawColor, SDPG_Foreground);
+						PDI->DrawLine(OldPos, NewPos, LineColor, SDPG_Foreground);
 
 						OldTime = NewTime;
 						OldPos = NewPos;
 					}
 				}
+
+				PDI->SetHitProxy(NULL);
 			}
 
 			OldKeyTime = NewKeyTime;
@@ -196,7 +202,7 @@ void FSplineComponentVisualizer::DrawVisualization(const UActorComponent* Compon
 	}
 }
 
-bool FSplineComponentVisualizer::VisProxyHandleClick(HComponentVisProxy* VisProxy)
+bool FSplineComponentVisualizer::VisProxyHandleClick(FLevelEditorViewportClient* InViewportClient, HComponentVisProxy* VisProxy, const FViewportClick& Click)
 {
 	bool bEditing = false;
 
@@ -213,6 +219,46 @@ bool FSplineComponentVisualizer::VisProxyHandleClick(HComponentVisProxy* VisProx
 			{
 				HSplineKeyProxy* KeyProxy = (HSplineKeyProxy*)VisProxy;
 				SelectedKeyIndex = KeyProxy->KeyIndex;
+				SelectedSegmentIndex = INDEX_NONE;
+				bEditing = true;
+			}
+			else if (VisProxy->IsA(HSplineSegmentProxy::StaticGetType()))
+			{
+				// Divide segment into subsegments and test each subsegment against ray representing click position and camera direction.
+				// Closest encounter with the spline determines the spline position.
+				const int32 NumSubdivisions = 16;
+
+				HSplineSegmentProxy* SegmentProxy = (HSplineSegmentProxy*)VisProxy;
+				SelectedKeyIndex = SegmentProxy->SegmentIndex;
+				SelectedSegmentIndex = SegmentProxy->SegmentIndex;
+
+				float SubsegmentStartKey = static_cast<float>(SelectedSegmentIndex);
+				FVector SubsegmentStart = SplineComp->ComponentToWorld.TransformPosition(SplineComp->SplineInfo.Eval(SubsegmentStartKey, FVector::ZeroVector));
+
+				float ClosestDistance = TNumericLimits<float>::Max();
+				FVector BestLocation = SubsegmentStart;
+
+				for (int32 Step = 1; Step < NumSubdivisions; Step++)
+				{
+					const float SubsegmentEndKey = SelectedSegmentIndex + Step / static_cast<float>(NumSubdivisions);
+					const FVector SubsegmentEnd = SplineComp->ComponentToWorld.TransformPosition(SplineComp->SplineInfo.Eval(SubsegmentEndKey, FVector::ZeroVector));
+
+					FVector SplineClosest;
+					FVector RayClosest;
+					FMath::SegmentDistToSegmentSafe(SubsegmentStart, SubsegmentEnd, Click.GetOrigin(), Click.GetOrigin() + Click.GetDirection() * 50000.0f, SplineClosest, RayClosest);
+
+					const float Distance = FVector::DistSquared(SplineClosest, RayClosest);
+					if (Distance < ClosestDistance)
+					{
+						ClosestDistance = Distance;
+						BestLocation = SplineClosest;
+					}
+
+					SubsegmentStartKey = SubsegmentEndKey;
+					SubsegmentStart = SubsegmentEnd;
+				}
+
+				SelectedSplinePosition = BestLocation;
 				bEditing = true;
 			}
 		}
@@ -231,7 +277,7 @@ USplineComponent* FSplineComponentVisualizer::GetEditedSplineComponent() const
 }
 
 
-bool FSplineComponentVisualizer::GetWidgetLocation(FVector& OutLocation) const
+bool FSplineComponentVisualizer::GetWidgetLocation(const FEditorViewportClient* ViewportClient, FVector& OutLocation) const
 {
 	USplineComponent* SplineComp = GetEditedSplineComponent();
 	if (SplineComp && SelectedKeyIndex != INDEX_NONE)
@@ -247,6 +293,32 @@ bool FSplineComponentVisualizer::GetWidgetLocation(FVector& OutLocation) const
 }
 
 
+bool FSplineComponentVisualizer::GetCustomInputCoordinateSystem(const FEditorViewportClient* ViewportClient, FMatrix& OutMatrix) const
+{
+	if (ViewportClient->GetWidgetCoordSystemSpace() == COORD_Local)
+	{
+		USplineComponent* SplineComp = GetEditedSplineComponent();
+		if (SplineComp && SelectedKeyIndex != INDEX_NONE)
+		{
+			if (SelectedKeyIndex < SplineComp->SplineInfo.Points.Num())
+			{
+				const auto& Point = SplineComp->SplineInfo.Points[SelectedKeyIndex];
+				const FVector Tangent = Point.ArriveTangent.SafeNormal();
+				const FVector Bitangent = (Tangent.Z == 1.0f) ?
+					FVector(1.0f, 0.0f, 0.0f) :
+					FVector(-Tangent.Y, Tangent.X, 0.0f).SafeNormal();
+				const FVector Normal = FVector::CrossProduct(Tangent, Bitangent);
+
+				OutMatrix = FMatrix(Tangent, Bitangent, Normal, FVector::ZeroVector) * FQuatRotationTranslationMatrix(SplineComp->ComponentToWorld.GetRotation(), FVector::ZeroVector);
+				return true;
+			}
+		}
+	}
+
+	return false;
+}
+
+
 bool FSplineComponentVisualizer::HandleInputDelta(FEditorViewportClient* ViewportClient, FViewport* Viewport, FVector& DeltaTranslate, FRotator& DeltaRotate, FVector& DeltaScale)
 {
 	USplineComponent* SplineComp = GetEditedSplineComponent();
@@ -283,7 +355,7 @@ bool FSplineComponentVisualizer::HandleInputDelta(FEditorViewportClient* Viewpor
 			EditedPoint.LeaveTangent = NewTangent;
 			EditedPoint.ArriveTangent = NewTangent;
 		}
-
+		 
 		if (!DeltaScale.IsZero())
 		{
 			// Set point tangent as user controlled
@@ -347,7 +419,7 @@ void FSplineComponentVisualizer::EndEditing()
 
 void FSplineComponentVisualizer::OnDuplicateKey()
 {
-	const FScopedTransaction Transaction(LOCTEXT("DuplicateKey", "Duplicate Spline Key"));
+	const FScopedTransaction Transaction(LOCTEXT("DuplicateSplinePoint", "Duplicate Spline Point"));
 	USplineComponent* SplineComp = GetEditedSplineComponent();
 
 	if (AActor* Owner = SplineComp->GetOwner())
@@ -368,9 +440,44 @@ void FSplineComponentVisualizer::OnDuplicateKey()
 }
 
 
+bool FSplineComponentVisualizer::CanAddKey() const
+{
+	USplineComponent* SplineComp = GetEditedSplineComponent();
+	return (SplineComp &&
+		SelectedSegmentIndex != INDEX_NONE &&
+		SelectedSegmentIndex < SplineComp->SplineInfo.Points.Num() - 1);
+}
+
+
+void FSplineComponentVisualizer::OnAddKey()
+{
+	const FScopedTransaction Transaction(LOCTEXT("AddSplinePoint", "Add Spline Point"));
+	USplineComponent* SplineComp = GetEditedSplineComponent();
+
+	if (AActor* Owner = SplineComp->GetOwner())
+	{
+		Owner->Modify();
+	}
+
+	FInterpCurvePoint<FVector> NewKey;
+	NewKey.InVal = SelectedSegmentIndex;
+	NewKey.OutVal = SplineComp->ComponentToWorld.InverseTransformPosition(SelectedSplinePosition);
+	NewKey.InterpMode = CIM_CurveAuto;
+	int32 NewKeyIndex = SplineComp->SplineInfo.Points.Insert(NewKey, SelectedSegmentIndex + 1);
+	// move selection to 'next' key
+	SelectedKeyIndex = SelectedSegmentIndex + 1;
+	SelectedSegmentIndex = INDEX_NONE;
+
+	// Update Input value for all keys
+	SplineComp->RefreshSplineInputs();
+
+	NotifyComponentModified();
+}
+
+
 void FSplineComponentVisualizer::OnDeleteKey()
 {
-	const FScopedTransaction Transaction(LOCTEXT("DeleteKey", "Delete Spline Key"));
+	const FScopedTransaction Transaction(LOCTEXT("DeleteSplinePoint", "Delete Spline Point"));
 	USplineComponent* SplineComp = GetEditedSplineComponent();
 
 	if (AActor* Owner = SplineComp->GetOwner())
@@ -382,15 +489,17 @@ void FSplineComponentVisualizer::OnDeleteKey()
 	SplineComp->RefreshSplineInputs(); // update input value for each key
 
 	SelectedKeyIndex = INDEX_NONE; // deselect any keys
+	SelectedSegmentIndex = INDEX_NONE;
 
 	NotifyComponentModified();
 }
 
 
-bool FSplineComponentVisualizer::IsSelectionValid() const
+bool FSplineComponentVisualizer::IsKeySelectionValid() const
 {
 	USplineComponent* SplineComp = GetEditedSplineComponent();
 	return (SplineComp &&
+		SelectedSegmentIndex == INDEX_NONE &&
 		SelectedKeyIndex != INDEX_NONE &&
 		SelectedKeyIndex < SplineComp->SplineInfo.Points.Num());
 }
@@ -431,7 +540,7 @@ void FSplineComponentVisualizer::OnSetKeyType(EInterpCurveMode Mode)
 
 	if (SplineComp->SplineInfo.Points[SelectedKeyIndex].InterpMode != Mode)
 	{
-		const FScopedTransaction Transaction(LOCTEXT("ResetToAutomaticTangent", "Reset to Automatic Tangent"));
+		const FScopedTransaction Transaction(LOCTEXT("SetSplinePointType", "Set Spline Point Type"));
 		if (AActor* Owner = SplineComp->GetOwner())
 		{
 			Owner->Modify();
@@ -445,7 +554,7 @@ void FSplineComponentVisualizer::OnSetKeyType(EInterpCurveMode Mode)
 
 bool FSplineComponentVisualizer::IsKeyTypeSet(EInterpCurveMode Mode) const
 {
-	if (IsSelectionValid())
+	if (IsKeySelectionValid())
 	{
 		USplineComponent* SplineComp = GetEditedSplineComponent();
 
@@ -463,25 +572,31 @@ TSharedPtr<SWidget> FSplineComponentVisualizer::GenerateContextMenu() const
 	{
 		MenuBuilder.BeginSection("SplineKeyEdit");
 		{
-			MenuBuilder.AddMenuEntry(FSplineComponentVisualizerCommands::Get().DeleteKey);
-			MenuBuilder.AddMenuEntry(FSplineComponentVisualizerCommands::Get().DuplicateKey);
-
-			MenuBuilder.AddSubMenu(
-				LOCTEXT("KeyType", "Key Type"),
-				LOCTEXT("KeyTypeTooltip", "Define the type of the spline key."),
-				FNewMenuDelegate::CreateSP(this, &FSplineComponentVisualizer::GenerateSplinePointTypeSubMenu));
-
-			// Only add the Automatic Tangents submenu if the key is a curve type
-			USplineComponent* SplineComp = GetEditedSplineComponent();
-			if (SplineComp &&
-				SelectedKeyIndex != INDEX_NONE &&
-				SelectedKeyIndex < SplineComp->SplineInfo.Points.Num() &&
-				SplineComp->SplineInfo.Points[SelectedKeyIndex].IsCurveKey())
+			if (SelectedSegmentIndex != INDEX_NONE)
+			{
+				MenuBuilder.AddMenuEntry(FSplineComponentVisualizerCommands::Get().AddKey);
+			}
+			else if (SelectedKeyIndex != INDEX_NONE)
 			{
+				MenuBuilder.AddMenuEntry(FSplineComponentVisualizerCommands::Get().DeleteKey);
+				MenuBuilder.AddMenuEntry(FSplineComponentVisualizerCommands::Get().DuplicateKey);
+
 				MenuBuilder.AddSubMenu(
-					LOCTEXT("ResetToAutomaticTangent", "Reset to Automatic Tangent"),
-					LOCTEXT("ResetToAutomaticTangentTooltip", "Reset the point tangent to an automatically generated value."),
-					FNewMenuDelegate::CreateSP(this, &FSplineComponentVisualizer::GenerateTangentTypeSubMenu));
+					LOCTEXT("SplinePointType", "Spline Point Type"),
+					LOCTEXT("KeyTypeTooltip", "Define the type of the spline point."),
+					FNewMenuDelegate::CreateSP(this, &FSplineComponentVisualizer::GenerateSplinePointTypeSubMenu));
+
+				// Only add the Automatic Tangents submenu if the key is a curve type
+				USplineComponent* SplineComp = GetEditedSplineComponent();
+				if (SplineComp &&
+					SelectedKeyIndex < SplineComp->SplineInfo.Points.Num() &&
+					SplineComp->SplineInfo.Points[SelectedKeyIndex].IsCurveKey())
+				{
+					MenuBuilder.AddSubMenu(
+						LOCTEXT("ResetToAutomaticTangent", "Reset to Automatic Tangent"),
+						LOCTEXT("ResetToAutomaticTangentTooltip", "Reset the spline point tangent to an automatically generated value."),
+						FNewMenuDelegate::CreateSP(this, &FSplineComponentVisualizer::GenerateTangentTypeSubMenu));
+				}
 			}
 		}
 		MenuBuilder.EndSection();
diff --git a/Engine/Source/Editor/ComponentVisualizers/Private/SplineComponentVisualizer.h b/Engine/Source/Editor/ComponentVisualizers/Private/SplineComponentVisualizer.h
--- a/Engine/Source/Editor/ComponentVisualizers/Private/SplineComponentVisualizer.h
+++ b/Engine/Source/Editor/ComponentVisualizers/Private/SplineComponentVisualizer.h
@@ -21,14 +21,26 @@ struct HSplineKeyProxy : public HSplineVisProxy
 	DECLARE_HIT_PROXY();
 
 	HSplineKeyProxy(const UActorComponent* InComponent, int32 InKeyIndex) 
-	: HSplineVisProxy(InComponent)
-	, KeyIndex(InKeyIndex)
+		: HSplineVisProxy(InComponent)
+		, KeyIndex(InKeyIndex)
 	{}
 
-
 	int32 KeyIndex;
 };
 
+/** Proxy for a spline segment */
+struct HSplineSegmentProxy : public HSplineVisProxy
+{
+	DECLARE_HIT_PROXY();
+
+	HSplineSegmentProxy(const UActorComponent* InComponent, int32 InSegmentIndex)
+		: HSplineVisProxy(InComponent)
+		, SegmentIndex(InSegmentIndex)
+	{}
+
+	int32 SegmentIndex;
+};
+
 /** SplineComponent visualizer/edit functionality */
 class FSplineComponentVisualizer : public FComponentVisualizer
 {
@@ -39,9 +51,10 @@ public:
 	// Begin FComponentVisualizer interface
 	virtual void OnRegister() override;
 	virtual void DrawVisualization(const UActorComponent* Component, const FSceneView* View, FPrimitiveDrawInterface* PDI) override;
-	virtual bool VisProxyHandleClick(HComponentVisProxy* VisProxy) override;
+	virtual bool VisProxyHandleClick(FLevelEditorViewportClient* InViewportClient, HComponentVisProxy* VisProxy, const FViewportClick& Click) override;
 	virtual void EndEditing() override;
-	virtual bool GetWidgetLocation(FVector& OutLocation) const override;
+	virtual bool GetWidgetLocation(const FEditorViewportClient* ViewportClient, FVector& OutLocation) const override;
+	virtual bool GetCustomInputCoordinateSystem(const FEditorViewportClient* ViewportClient, FMatrix& OutMatrix) const override;
 	virtual bool HandleInputDelta(FEditorViewportClient* ViewportClient, FViewport* Viewport, FVector& DeltaTranslate, FRotator& DeltaRotate, FVector& DeltaScale) override;
 	virtual bool HandleInputKey(FEditorViewportClient* ViewportClient, FViewport* Viewport, FKey Key, EInputEvent Event) override;
 	virtual TSharedPtr<SWidget> GenerateContextMenu() const override;
@@ -56,7 +69,10 @@ private:
 
 	void OnDeleteKey();
 	void OnDuplicateKey();
-	bool IsSelectionValid() const;
+	bool IsKeySelectionValid() const;
+
+	void OnAddKey();
+	bool CanAddKey() const;
 
 	void OnResetToAutomaticTangent(EInterpCurveMode Mode);
 	bool CanResetToAutomaticTangent(EInterpCurveMode Mode) const;
@@ -75,11 +91,19 @@ private:
 
 	/** Actor that owns the currently edited spline */
 	TWeakObjectPtr<AActor> SplineOwningActor;
+
 	/** Name of property on the actor that references the spline we are editing */
 	FName SplineCompPropName;
+
 	/** Index of key we have selected */
 	int32 SelectedKeyIndex;
 
+	/** Index of segment we have selected */
+	int32 SelectedSegmentIndex;
+
+	/** Position on spline we have selected */
+	FVector SelectedSplinePosition;
+
 	/** Whether we currently allow duplication when dragging */
 	bool bAllowDuplication;
 };
diff --git a/Engine/Source/Editor/UnrealEd/Private/ComponentVisualizer.cpp b/Engine/Source/Editor/UnrealEd/Private/ComponentVisualizer.cpp
--- a/Engine/Source/Editor/UnrealEd/Private/ComponentVisualizer.cpp
+++ b/Engine/Source/Editor/UnrealEd/Private/ComponentVisualizer.cpp
@@ -15,8 +15,8 @@ FName FComponentVisualizer::GetComponentPropertyName(const UActorComponent* Comp
 		if(CompOwner)
 		{
 			// Iterate over fields of this actor
-			UClass* SplineActorClass = CompOwner->GetClass();
-			for (TFieldIterator<UObjectProperty> It(SplineActorClass); It; ++It)
+			UClass* ActorClass = CompOwner->GetClass();
+			for (TFieldIterator<UObjectProperty> It(ActorClass); It; ++It)
 			{
 				// See if this property points to the component in question
 				UObjectProperty* ObjectProp = *It;
diff --git a/Engine/Source/Editor/UnrealEd/Private/ComponentVisualizerManager.cpp b/Engine/Source/Editor/UnrealEd/Private/ComponentVisualizerManager.cpp
--- a/Engine/Source/Editor/UnrealEd/Private/ComponentVisualizerManager.cpp
+++ b/Engine/Source/Editor/UnrealEd/Private/ComponentVisualizerManager.cpp
@@ -4,20 +4,20 @@
 #include "ComponentVisualizerManager.h"
 #include "ILevelEditor.h"
 
-/** Handle a click on the specified level editor viewport client */
-bool FComponentVisualizerManager::HandleClick(FLevelEditorViewportClient* InViewportClient, HHitProxy *HitProxy, const FViewportClick &Click)
+/** Handle a click on the specified editor viewport client */
+bool FComponentVisualizerManager::HandleClick(FLevelEditorViewportClient* InViewportClient, HHitProxy* HitProxy, const FViewportClick& Click)
 {
-	bool bHandled = HandleProxyForComponentVis(HitProxy);
+	bool bHandled = HandleProxyForComponentVis(InViewportClient, HitProxy, Click);
 	if (bHandled && Click.GetKey() == EKeys::RightMouseButton)
 	{
 		TSharedPtr<SWidget> MenuWidget = GenerateContextMenuForComponentVis();
 		if (MenuWidget.IsValid())
 		{
-			auto ParentLevelEditorPinned = InViewportClient->ParentLevelEditor.Pin();
-			if (ParentLevelEditorPinned.IsValid())
+			auto ParentLevelEditor = InViewportClient->ParentLevelEditor.Pin();
+			if (ParentLevelEditor.IsValid())
 			{
 				FSlateApplication::Get().PushMenu(
-					ParentLevelEditorPinned.ToSharedRef(),
+					ParentLevelEditor.ToSharedRef(),
 					MenuWidget.ToSharedRef(),
 					FSlateApplication::Get().GetCursorPos(),
 					FPopupTransitionEffect(FPopupTransitionEffect::ContextMenu));
@@ -30,7 +30,7 @@ bool FComponentVisualizerManager::HandleClick(FLevelEditorViewportClient* InView
 	return false;
 }
 
-bool FComponentVisualizerManager::HandleProxyForComponentVis(HHitProxy *HitProxy)
+bool FComponentVisualizerManager::HandleProxyForComponentVis(FLevelEditorViewportClient* InViewportClient, HHitProxy* HitProxy, const FViewportClick& Click)
 {
 	if (HitProxy->IsA(HComponentVisProxy::StaticGetType()))
 	{
@@ -41,16 +41,17 @@ bool FComponentVisualizerManager::HandleProxyForComponentVis(HHitProxy *HitProxy
 			TSharedPtr<FComponentVisualizer> Visualizer = GUnrealEd->FindComponentVisualizer(ClickedComponent->GetClass());
 			if (Visualizer.IsValid())
 			{
-				bool bIsActive = Visualizer->VisProxyHandleClick(VisProxy);
+				bool bIsActive = Visualizer->VisProxyHandleClick(InViewportClient, VisProxy, Click);
 				if (bIsActive)
 				{
 					// call EndEditing on any currently edited visualizer, if we are going to change it
-					if (EditedVisualizer.IsValid() && Visualizer.Get() != EditedVisualizer.Pin().Get())
+					TSharedPtr<FComponentVisualizer> EditedVisualizer = EditedVisualizerPtr.Pin();
+					if (EditedVisualizer.IsValid() && Visualizer.Get() != EditedVisualizer.Get())
 					{
-						EditedVisualizer.Pin()->EndEditing();
+						EditedVisualizer->EndEditing();
 					}
 
-					EditedVisualizer = Visualizer;
+					EditedVisualizerPtr = Visualizer;
 					return true;
 				}
 			}
@@ -66,18 +67,22 @@ bool FComponentVisualizerManager::HandleProxyForComponentVis(HHitProxy *HitProxy
 
 void FComponentVisualizerManager::ClearActiveComponentVis()
 {
+	TSharedPtr<FComponentVisualizer> EditedVisualizer = EditedVisualizerPtr.Pin();
+
 	if (EditedVisualizer.IsValid())
 	{
-		EditedVisualizer.Pin()->EndEditing();
-		EditedVisualizer = NULL;
+		EditedVisualizer->EndEditing();
+		EditedVisualizerPtr = NULL;
 	}
 }
 
 bool FComponentVisualizerManager::HandleInputKey(FEditorViewportClient* ViewportClient, FViewport* Viewport, FKey Key, EInputEvent Event) const
 {
-	if(EditedVisualizer.IsValid())
+	TSharedPtr<FComponentVisualizer> EditedVisualizer = EditedVisualizerPtr.Pin();
+
+	if (EditedVisualizer.IsValid())
 	{
-		if(EditedVisualizer.Pin()->HandleInputKey(ViewportClient, Viewport, Key, Event))
+		if(EditedVisualizer->HandleInputKey(ViewportClient, Viewport, Key, Event))
 		{
 			return true;
 		}
@@ -88,9 +93,11 @@ bool FComponentVisualizerManager::HandleInputKey(FEditorViewportClient* Viewport
 
 bool FComponentVisualizerManager::HandleInputDelta(FEditorViewportClient* InViewportClient, FViewport* InViewport, FVector& InDrag, FRotator& InRot, FVector& InScale) const
 {
+	TSharedPtr<FComponentVisualizer> EditedVisualizer = EditedVisualizerPtr.Pin();
+
 	if (EditedVisualizer.IsValid() && InViewportClient->GetCurrentWidgetAxis() != EAxisList::None)
 	{
-		if (EditedVisualizer.Pin()->HandleInputDelta(InViewportClient, InViewport, InDrag, InRot, InScale))
+		if (EditedVisualizer->HandleInputDelta(InViewportClient, InViewport, InDrag, InRot, InScale))
 		{
 			return true;
 		}
@@ -100,11 +107,26 @@ bool FComponentVisualizerManager::HandleInputDelta(FEditorViewportClient* InView
 }
 
 
-bool FComponentVisualizerManager::GetWidgetLocation(FVector& OutLocation) const
+bool FComponentVisualizerManager::GetWidgetLocation(const FEditorViewportClient* ViewportClient, FVector& OutLocation) const
+{
+	TSharedPtr<FComponentVisualizer> EditedVisualizer = EditedVisualizerPtr.Pin();
+
+	if (EditedVisualizer.IsValid())
+	{
+		return EditedVisualizer->GetWidgetLocation(ViewportClient, OutLocation);
+	}
+
+	return false;
+}
+
+
+bool FComponentVisualizerManager::GetCustomInputCoordinateSystem(const FEditorViewportClient* ViewportClient, FMatrix& OutMatrix) const
 {
+	TSharedPtr<FComponentVisualizer> EditedVisualizer = EditedVisualizerPtr.Pin();
+
 	if (EditedVisualizer.IsValid())
 	{
-		return EditedVisualizer.Pin()->GetWidgetLocation(OutLocation);
+		return EditedVisualizer->GetCustomInputCoordinateSystem(ViewportClient, OutMatrix);
 	}
 
 	return false;
@@ -113,9 +135,11 @@ bool FComponentVisualizerManager::GetWidgetLocation(FVector& OutLocation) const
 
 TSharedPtr<SWidget> FComponentVisualizerManager::GenerateContextMenuForComponentVis() const
 {
+	TSharedPtr<FComponentVisualizer> EditedVisualizer = EditedVisualizerPtr.Pin();
+
 	if (EditedVisualizer.IsValid())
 	{
-		return EditedVisualizer.Pin()->GenerateContextMenu();
+		return EditedVisualizer->GenerateContextMenu();
 	}
 
 	return TSharedPtr<SWidget>();
diff --git a/Engine/Source/Editor/UnrealEd/Private/LevelEditorViewport.cpp b/Engine/Source/Editor/UnrealEd/Private/LevelEditorViewport.cpp
--- a/Engine/Source/Editor/UnrealEd/Private/LevelEditorViewport.cpp
+++ b/Engine/Source/Editor/UnrealEd/Private/LevelEditorViewport.cpp
@@ -3069,7 +3069,7 @@ FWidget::EWidgetMode FLevelEditorViewportClient::GetWidgetMode() const
 FVector FLevelEditorViewportClient::GetWidgetLocation() const
 {
 	FVector ComponentVisWidgetLocation;
-	if (GUnrealEd->ComponentVisManager.GetWidgetLocation(ComponentVisWidgetLocation))
+	if (GUnrealEd->ComponentVisManager.GetWidgetLocation(this, ComponentVisWidgetLocation))
 	{
 		return ComponentVisWidgetLocation;
 	}
@@ -3079,6 +3079,12 @@ FVector FLevelEditorViewportClient::GetWidgetLocation() const
 
 FMatrix FLevelEditorViewportClient::GetWidgetCoordSystem() const 
 {
+	FMatrix ComponentVisWidgetCoordSystem;
+	if (GUnrealEd->ComponentVisManager.GetCustomInputCoordinateSystem(this, ComponentVisWidgetCoordSystem))
+	{
+		return ComponentVisWidgetCoordSystem;
+	}
+
 	return GLevelEditorModeTools().GetCustomInputCoordinateSystem();
 }
 
diff --git a/Engine/Source/Editor/UnrealEd/Public/ComponentVisualizer.h b/Engine/Source/Editor/UnrealEd/Public/ComponentVisualizer.h
--- a/Engine/Source/Editor/UnrealEd/Public/ComponentVisualizer.h
+++ b/Engine/Source/Editor/UnrealEd/Public/ComponentVisualizer.h
@@ -33,11 +33,13 @@ public:
 	/** Draw visualization for the supplied component */
 	virtual void DrawVisualization(const UActorComponent* Component, const FSceneView* View, FPrimitiveDrawInterface* PDI) {}
 	/** */
-	virtual bool VisProxyHandleClick(HComponentVisProxy* VisProxy) { return false; }
+	virtual bool VisProxyHandleClick(FLevelEditorViewportClient* InViewportClient, HComponentVisProxy* VisProxy, const FViewportClick& Click) { return false; }
 	/** */
 	virtual void EndEditing() {}
 	/** */
-	virtual bool GetWidgetLocation(FVector& OutLocation) const { return false; }
+	virtual bool GetWidgetLocation(const FEditorViewportClient* ViewportClient, FVector& OutLocation) const { return false; }
+	/** */
+	virtual bool GetCustomInputCoordinateSystem(const FEditorViewportClient* ViewportClient, FMatrix& OutMatrix) const { return false; }
 	/** */
 	virtual bool HandleInputDelta(FEditorViewportClient* ViewportClient, FViewport* Viewport, FVector& DeltaTranslate, FRotator& DeltalRotate, FVector& DeltaScale) { return false; }
 	/** */
diff --git a/Engine/Source/Editor/UnrealEd/Public/ComponentVisualizerManager.h b/Engine/Source/Editor/UnrealEd/Public/ComponentVisualizerManager.h
--- a/Engine/Source/Editor/UnrealEd/Public/ComponentVisualizerManager.h
+++ b/Engine/Source/Editor/UnrealEd/Public/ComponentVisualizerManager.h
@@ -15,23 +15,25 @@ public:
 
 
 	/** Activate a component visualizer given a clicked proxy */
-	bool HandleProxyForComponentVis(HHitProxy *HitProxy);
+	bool HandleProxyForComponentVis(FLevelEditorViewportClient* InViewportClient, HHitProxy *HitProxy, const FViewportClick &Click);
 	/** Clear active component visualizer */
 	void ClearActiveComponentVis();
 
 	/** Handle a click on the specified level editor viewport client */
 	bool HandleClick(FLevelEditorViewportClient* InViewportClient, HHitProxy *HitProxy, const FViewportClick &Click);
 	/** Pass key input to active visualizer */
-	bool HandleInputKey(FEditorViewportClient* ViewportClient, FViewport* Viewport, FKey Key, EInputEvent Event) const;
+	bool HandleInputKey(FEditorViewportClient* InViewportClient, FViewport* Viewport, FKey Key, EInputEvent Event) const;
 	/** Pass delta input to active visualizer */
 	bool HandleInputDelta(FEditorViewportClient* InViewportClient, FViewport* InViewport, FVector& InDrag, FRotator& InRot, FVector& InScale) const;
 	/** Get widget location from active visualizer */
-	bool GetWidgetLocation(FVector& OutLocation) const;
+	bool GetWidgetLocation(const FEditorViewportClient* InViewportClient, FVector& OutLocation) const;
+	/** Get custom widget coordinate system from active visualizer */
+	bool GetCustomInputCoordinateSystem(const FEditorViewportClient* InViewportClient, FMatrix& OutMatrix) const;
 
 	/** Generate context menu for the component visualizer */
 	TSharedPtr<SWidget> GenerateContextMenuForComponentVis() const;
 
 private:
 	/** Currently 'active' visualizer that we should pass input to etc */
-	TWeakPtr<class FComponentVisualizer> EditedVisualizer;
+	TWeakPtr<class FComponentVisualizer> EditedVisualizerPtr;
 };
\ No newline at end of file
