Collect all potential attenuation shapes for an audio component correctly

diff --git a/Engine/Source/Editor/ComponentVisualizers/Private/AudioComponentVisualizer.cpp b/Engine/Source/Editor/ComponentVisualizers/Private/AudioComponentVisualizer.cpp
--- a/Engine/Source/Editor/ComponentVisualizers/Private/AudioComponentVisualizer.cpp
+++ b/Engine/Source/Editor/ComponentVisualizers/Private/AudioComponentVisualizer.cpp
@@ -15,7 +15,7 @@ void FAudioComponentVisualizer::DrawVisualization( const UActorComponent* Compon
 		{
 			const FTransform& Transform = AudioComp->ComponentToWorld;
 
-			TMap<EAttenuationShape::Type, FAttenuationSettings::AttenuationShapeDetails> ShapeDetailsMap;
+			TMultiMap<EAttenuationShape::Type, FAttenuationSettings::AttenuationShapeDetails> ShapeDetailsMap;
 			AudioComp->CollectAttenuationShapesForVisualization(ShapeDetailsMap);
 
 			for ( auto It = ShapeDetailsMap.CreateConstIterator(); It; ++It )
diff --git a/Engine/Source/Runtime/Engine/Classes/Components/AudioComponent.h b/Engine/Source/Runtime/Engine/Classes/Components/AudioComponent.h
--- a/Engine/Source/Runtime/Engine/Classes/Components/AudioComponent.h
+++ b/Engine/Source/Runtime/Engine/Classes/Components/AudioComponent.h
@@ -309,7 +309,7 @@ public:
 	// End ActorComponent interface.
 
 	const FAttenuationSettings* GetAttenuationSettingsToApply() const;
-	void CollectAttenuationShapesForVisualization(TMap<EAttenuationShape::Type, FAttenuationSettings::AttenuationShapeDetails>& ShapeDetailsMap) const;
+	void CollectAttenuationShapesForVisualization(TMultiMap<EAttenuationShape::Type, FAttenuationSettings::AttenuationShapeDetails>& ShapeDetailsMap) const;
 
 private:
 	
diff --git a/Engine/Source/Runtime/Engine/Classes/Sound/SoundAttenuation.h b/Engine/Source/Runtime/Engine/Classes/Sound/SoundAttenuation.h
--- a/Engine/Source/Runtime/Engine/Classes/Sound/SoundAttenuation.h
+++ b/Engine/Source/Runtime/Engine/Classes/Sound/SoundAttenuation.h
@@ -137,7 +137,7 @@ struct ENGINE_API FAttenuationSettings
 		float ConeOffset;
 	};
 
-	void CollectAttenuationShapesForVisualization(TMap<EAttenuationShape::Type, AttenuationShapeDetails>& ShapeDetailsMap) const;
+	void CollectAttenuationShapesForVisualization(TMultiMap<EAttenuationShape::Type, AttenuationShapeDetails>& ShapeDetailsMap) const;
 	float GetMaxDimension() const;
 
 private:
diff --git a/Engine/Source/Runtime/Engine/Private/ActiveSound.cpp b/Engine/Source/Runtime/Engine/Private/ActiveSound.cpp
--- a/Engine/Source/Runtime/Engine/Private/ActiveSound.cpp
+++ b/Engine/Source/Runtime/Engine/Private/ActiveSound.cpp
@@ -558,7 +558,7 @@ void FActiveSound::SetIntParameter( FName InName, const int32 InInt )
 	}
 }
 
-void FActiveSound::CollectAttenuationShapesForVisualization(TMap<EAttenuationShape::Type, FAttenuationSettings::AttenuationShapeDetails>& ShapeDetailsMap) const
+void FActiveSound::CollectAttenuationShapesForVisualization(TMultiMap<EAttenuationShape::Type, FAttenuationSettings::AttenuationShapeDetails>& ShapeDetailsMap) const
 {
 	bool bFoundAttenuationSettings = false;
 
diff --git a/Engine/Source/Runtime/Engine/Private/AudioComponent.cpp b/Engine/Source/Runtime/Engine/Private/AudioComponent.cpp
--- a/Engine/Source/Runtime/Engine/Private/AudioComponent.cpp
+++ b/Engine/Source/Runtime/Engine/Private/AudioComponent.cpp
@@ -359,7 +359,7 @@ const FAttenuationSettings* UAudioComponent::GetAttenuationSettingsToApply() con
 	return NULL;
 }
 
-void UAudioComponent::CollectAttenuationShapesForVisualization(TMap<EAttenuationShape::Type, FAttenuationSettings::AttenuationShapeDetails>& ShapeDetailsMap) const
+void UAudioComponent::CollectAttenuationShapesForVisualization(TMultiMap<EAttenuationShape::Type, FAttenuationSettings::AttenuationShapeDetails>& ShapeDetailsMap) const
 {
 	const FAttenuationSettings *AttenuationSettingsToApply = GetAttenuationSettingsToApply();
 
diff --git a/Engine/Source/Runtime/Engine/Private/SoundAttenuation.cpp b/Engine/Source/Runtime/Engine/Private/SoundAttenuation.cpp
--- a/Engine/Source/Runtime/Engine/Private/SoundAttenuation.cpp
+++ b/Engine/Source/Runtime/Engine/Private/SoundAttenuation.cpp
@@ -235,7 +235,7 @@ void FAttenuationSettings::ApplyAttenuation( const FTransform& SoundTransform, c
 	}
 }
 
-void FAttenuationSettings::CollectAttenuationShapesForVisualization(TMap<EAttenuationShape::Type, AttenuationShapeDetails>& ShapeDetailsMap) const
+void FAttenuationSettings::CollectAttenuationShapesForVisualization(TMultiMap<EAttenuationShape::Type, AttenuationShapeDetails>& ShapeDetailsMap) const
 {
 	if (bAttenuate)
 	{
diff --git a/Engine/Source/Runtime/Engine/Private/UnrealEngine.cpp b/Engine/Source/Runtime/Engine/Private/UnrealEngine.cpp
--- a/Engine/Source/Runtime/Engine/Private/UnrealEngine.cpp
+++ b/Engine/Source/Runtime/Engine/Private/UnrealEngine.cpp
@@ -10982,7 +10982,7 @@ int32 UEngine::RenderStatSounds(UWorld* World, FViewport* Viewport, FCanvas* Can
 					float SphereRadius = 0.f;
 					float SphereInnerRadius = 0.f;
 
-					TMap<EAttenuationShape::Type, FAttenuationSettings::AttenuationShapeDetails> ShapeDetailsMap;
+					TMultiMap<EAttenuationShape::Type, FAttenuationSettings::AttenuationShapeDetails> ShapeDetailsMap;
 					ActiveSound.CollectAttenuationShapesForVisualization(ShapeDetailsMap);
 
 					if (ShapeDetailsMap.Num() > 0)
diff --git a/Engine/Source/Runtime/Engine/Public/ActiveSound.h b/Engine/Source/Runtime/Engine/Public/ActiveSound.h
--- a/Engine/Source/Runtime/Engine/Public/ActiveSound.h
+++ b/Engine/Source/Runtime/Engine/Public/ActiveSound.h
@@ -176,7 +176,7 @@ public:
 	 */
 	int32 GetIntParameter(const FName InName, int32& OutInt) const;
 
-	void CollectAttenuationShapesForVisualization(TMap<EAttenuationShape::Type, FAttenuationSettings::AttenuationShapeDetails>& ShapeDetailsMap) const;
+	void CollectAttenuationShapesForVisualization(TMultiMap<EAttenuationShape::Type, FAttenuationSettings::AttenuationShapeDetails>& ShapeDetailsMap) const;
 
 	/**
 	 * Friend archive function used for serialization.
