[GitHub] 323: New SensingComponentVisualizer

by timdonks

diff --git a/Engine/Source/Editor/ComponentVisualizers/ComponentVisualizers.Build.cs b/Engine/Source/Editor/ComponentVisualizers/ComponentVisualizers.Build.cs
--- a/Engine/Source/Editor/ComponentVisualizers/ComponentVisualizers.Build.cs
+++ b/Engine/Source/Editor/ComponentVisualizers/ComponentVisualizers.Build.cs
@@ -18,7 +18,8 @@ public class ComponentVisualizers : ModuleRules
                 "SlateCore",
 				"UnrealEd",
                 "PropertyEditor",
-                "EditorStyle"
+                "EditorStyle",
+				"AIModule"
 			}
 			);
 	}
diff --git a/Engine/Source/Editor/ComponentVisualizers/Private/ComponentVisualizers.cpp b/Engine/Source/Editor/ComponentVisualizers/Private/ComponentVisualizers.cpp
--- a/Engine/Source/Editor/ComponentVisualizers/Private/ComponentVisualizers.cpp
+++ b/Engine/Source/Editor/ComponentVisualizers/Private/ComponentVisualizers.cpp
@@ -4,6 +4,7 @@
 #include "ComponentVisualizers.h"
 
 #include "SoundDefinitions.h"
+#include "Perception/PawnSensingComponent.h"
 
 #include "PointLightComponentVisualizer.h"
 #include "SpotLightComponentVisualizer.h"
@@ -12,6 +13,7 @@
 #include "ConstraintComponentVisualizer.h"
 #include "SpringArmComponentVisualizer.h"
 #include "SplineComponentVisualizer.h"
+#include "SensingComponentVisualizer.h"
 
 IMPLEMENT_MODULE( FComponentVisualizersModule, ComponentVisualizers );
 
@@ -24,6 +26,7 @@ void FComponentVisualizersModule::StartupModule()
 	RegisterComponentVisualizer(UPhysicsConstraintComponent::StaticClass()->GetFName(), MakeShareable(new FConstraintComponentVisualizer));
 	RegisterComponentVisualizer(USpringArmComponent::StaticClass()->GetFName(), MakeShareable(new FSpringArmComponentVisualizer));
 	RegisterComponentVisualizer(USplineComponent::StaticClass()->GetFName(), MakeShareable(new FSplineComponentVisualizer));
+	RegisterComponentVisualizer(UPawnSensingComponent::StaticClass()->GetFName(), MakeShareable(new FSensingComponentVisualizer));
 }
 
 void FComponentVisualizersModule::ShutdownModule()
diff --git a/Engine/Source/Editor/ComponentVisualizers/Private/SensingComponentVisualizer.cpp b/Engine/Source/Editor/ComponentVisualizers/Private/SensingComponentVisualizer.cpp
new file mode 100644
--- /dev/null
+++ b/Engine/Source/Editor/ComponentVisualizers/Private/SensingComponentVisualizer.cpp
@@ -0,0 +1,36 @@
+// Copyright 1998-2014 Epic Games, Inc. All Rights Reserved.
+
+#include "ComponentVisualizersPrivatePCH.h"
+#include "SensingComponentVisualizer.h"
+#include "Perception/PawnSensingComponent.h"
+
+
+void FSensingComponentVisualizer::DrawVisualization(const UActorComponent* Component, const FSceneView* View, FPrimitiveDrawInterface* PDI)
+{
+	if (View->Family->EngineShowFlags.VisualizeSenses)
+	{
+		const UPawnSensingComponent* Senses = Cast<const UPawnSensingComponent>(Component);
+		if (Senses != NULL)
+		{
+			TArray<FVector> Verts;
+
+			//LOS hearing
+			if (Senses->LOSHearingThreshold > 0.0f)
+			{
+				DrawWireSphere(PDI, Senses->GetSensorLocation(), FColor::Yellow, Senses->LOSHearingThreshold, 16, SDPG_World);
+			}
+
+			//Hearing
+			if (Senses->HearingThreshold > 0.0f)
+			{
+				DrawWireSphere(PDI, Senses->GetSensorLocation(), FColor::Cyan, Senses->HearingThreshold, 16, SDPG_World);
+			}
+
+			// Sight
+			if (Senses->SightRadius > 0.0f)
+			{
+				DrawWireCone(PDI, FTranslationMatrix(Senses->GetSensorLocation()), Senses->SightRadius, Senses->GetPeripheralVisionAngle(), 10, FColor::Green, SDPG_World, Verts);
+			}
+		}
+	}
+}
diff --git a/Engine/Source/Editor/ComponentVisualizers/Private/SensingComponentVisualizer.h b/Engine/Source/Editor/ComponentVisualizers/Private/SensingComponentVisualizer.h
new file mode 100644
--- /dev/null
+++ b/Engine/Source/Editor/ComponentVisualizers/Private/SensingComponentVisualizer.h
@@ -0,0 +1,13 @@
+// Copyright 1998-2014 Epic Games, Inc. All Rights Reserved.
+
+#pragma once
+
+#include "ComponentVisualizer.h"
+
+class FSensingComponentVisualizer : public FComponentVisualizer
+{
+public:
+	// Begin FComponentVisualizer interface
+	virtual void DrawVisualization(const UActorComponent* Component, const FSceneView* View, FPrimitiveDrawInterface* PDI) override;
+	// End FComponentVisualizer interface
+};
diff --git a/Engine/Source/Runtime/Engine/Public/ShowFlagsValues.inl b/Engine/Source/Runtime/Engine/Public/ShowFlagsValues.inl
--- a/Engine/Source/Runtime/Engine/Public/ShowFlagsValues.inl
+++ b/Engine/Source/Runtime/Engine/Public/ShowFlagsValues.inl
@@ -272,7 +272,8 @@ SHOWFLAG_ALWAYS_ACCESSIBLE(DistanceFieldAO, SFG_LightingFeatures, LOCTEXT("Dista
 SHOWFLAG_FIXED_IN_SHIPPING(VisualizeSSR, 0, SFG_Visualize, LOCTEXT("VisualizeSSR", "Visualize ScreenSpaceReflections"))
 /** Force the use of the GBuffer. */
 SHOWFLAG_ALWAYS_ACCESSIBLE(ForceGBuffer, SFG_Hidden, LOCTEXT("ForceGBuffer", "Force usage of GBuffer"))
-
+/** Visualize the senses of the Senses component */
+SHOWFLAG_ALWAYS_ACCESSIBLE(VisualizeSenses, SFG_Advanced, LOCTEXT("VisualizeSenses", "Senses"))
 
 #undef SHOWFLAG_ALWAYS_ACCESSIBLE
 #undef SHOWFLAG_FIXED_IN_SHIPPING
