[GitHub] 435: SensingComponentVisualizer fix

by timdonks
(with minor refactoring)

diff --git a/Engine/Source/Editor/ComponentVisualizers/Private/SensingComponentVisualizer.cpp b/Engine/Source/Editor/ComponentVisualizers/Private/SensingComponentVisualizer.cpp
--- a/Engine/Source/Editor/ComponentVisualizers/Private/SensingComponentVisualizer.cpp
+++ b/Engine/Source/Editor/ComponentVisualizers/Private/SensingComponentVisualizer.cpp
@@ -12,24 +12,25 @@ void FSensingComponentVisualizer::DrawVisualization(const UActorComponent* Compo
 		const UPawnSensingComponent* Senses = Cast<const UPawnSensingComponent>(Component);
 		if (Senses != NULL)
 		{
-			TArray<FVector> Verts;
+			const FTransform Transform = FTransform(Senses->GetSensorRotation(), Senses->GetSensorLocation());
 
 			//LOS hearing
 			if (Senses->LOSHearingThreshold > 0.0f)
 			{
-				DrawWireSphere(PDI, Senses->GetSensorLocation(), FColor::Yellow, Senses->LOSHearingThreshold, 16, SDPG_World);
+				DrawWireSphere(PDI, Transform, FColor::Yellow, Senses->LOSHearingThreshold, 16, SDPG_World);
 			}
 
 			//Hearing
 			if (Senses->HearingThreshold > 0.0f)
 			{
-				DrawWireSphere(PDI, Senses->GetSensorLocation(), FColor::Cyan, Senses->HearingThreshold, 16, SDPG_World);
+				DrawWireSphere(PDI, Transform, FColor::Cyan, Senses->HearingThreshold, 16, SDPG_World);
 			}
 
 			// Sight
 			if (Senses->SightRadius > 0.0f)
 			{
-				DrawWireCone(PDI, FTranslationMatrix(Senses->GetSensorLocation()), Senses->SightRadius, Senses->GetPeripheralVisionAngle(), 10, FColor::Green, SDPG_World, Verts);
+				TArray<FVector> Verts;
+				DrawWireCone(PDI, Transform, Senses->SightRadius, Senses->GetPeripheralVisionAngle(), 10, FColor::Green, SDPG_World, Verts);
 			}
 		}
 	}
diff --git a/Engine/Source/Runtime/AIModule/Classes/Perception/PawnSensingComponent.h b/Engine/Source/Runtime/AIModule/Classes/Perception/PawnSensingComponent.h
--- a/Engine/Source/Runtime/AIModule/Classes/Perception/PawnSensingComponent.h
+++ b/Engine/Source/Runtime/AIModule/Classes/Perception/PawnSensingComponent.h
@@ -123,6 +123,9 @@ public:
 	/** Get position where hearing/seeing occurs (i.e. ear/eye position).  If we ever need different positions for hearing/seeing, we'll deal with that then! */
 	virtual FVector GetSensorLocation() const;
 
+	/**  Get the rotation of this sensor. We need this for the sight component */
+	virtual FRotator GetSensorRotation() const;
+
 protected:
 
 	/** See if there are interesting sounds and sights that we want to detect, and respond to them if so. */
diff --git a/Engine/Source/Runtime/AIModule/Private/Perception/PawnSensingComponent.cpp b/Engine/Source/Runtime/AIModule/Private/Perception/PawnSensingComponent.cpp
--- a/Engine/Source/Runtime/AIModule/Private/Perception/PawnSensingComponent.cpp
+++ b/Engine/Source/Runtime/AIModule/Private/Perception/PawnSensingComponent.cpp
@@ -386,6 +386,18 @@ FVector UPawnSensingComponent::GetSensorLocation() const
 	return SensorLocation;
 }
 
+FRotator UPawnSensingComponent::GetSensorRotation() const
+{
+	FRotator SensorRotation(FRotator::ZeroRotator);
+	const AActor* SensorActor = GetSensorActor();
+	
+	if (SensorActor != NULL)
+	{
+		SensorRotation = SensorActor->GetActorRotation();
+	}
+	
+	return SensorRotation;
+}
 
 bool UPawnSensingComponent::CanHear(const FVector& NoiseLoc, float Loudness, bool bFailedLOS) const
 {
