Add COM visualization

diff --git a/Engine/Source/Editor/ComponentVisualizers/Private/ComponentVisualizers.cpp b/Engine/Source/Editor/ComponentVisualizers/Private/ComponentVisualizers.cpp
--- a/Engine/Source/Editor/ComponentVisualizers/Private/ComponentVisualizers.cpp
+++ b/Engine/Source/Editor/ComponentVisualizers/Private/ComponentVisualizers.cpp
@@ -16,6 +16,7 @@
 #include "SplineComponentVisualizer.h"
 #include "SensingComponentVisualizer.h"
 #include "SpringComponentVisualizer.h"
+#include "PrimitiveComponentVisualizer.h"
 
 IMPLEMENT_MODULE( FComponentVisualizersModule, ComponentVisualizers );
 
@@ -30,6 +31,7 @@ void FComponentVisualizersModule::StartupModule()
 	RegisterComponentVisualizer(USplineComponent::StaticClass()->GetFName(), MakeShareable(new FSplineComponentVisualizer));
 	RegisterComponentVisualizer(UPawnSensingComponent::StaticClass()->GetFName(), MakeShareable(new FSensingComponentVisualizer));
 	RegisterComponentVisualizer(UPhysicsSpringComponent::StaticClass()->GetFName(), MakeShareable(new FSpringComponentVisualizer));
+	RegisterComponentVisualizer(UPrimitiveComponent::StaticClass()->GetFName(), MakeShareable(new FPrimitiveComponentVisualizer));
 }
 
 void FComponentVisualizersModule::ShutdownModule()
diff --git a/Engine/Source/Editor/ComponentVisualizers/Private/PrimitiveComponentVisualizer.cpp b/Engine/Source/Editor/ComponentVisualizers/Private/PrimitiveComponentVisualizer.cpp
new file mode 100644
--- /dev/null
+++ b/Engine/Source/Editor/ComponentVisualizers/Private/PrimitiveComponentVisualizer.cpp
@@ -0,0 +1,23 @@
+// Copyright 1998-2014 Epic Games, Inc. All Rights Reserved.
+
+#include "ComponentVisualizersPrivatePCH.h"
+
+#include "PrimitiveComponentVisualizer.h"
+
+static const FColor	COMColor(0,255,0);
+
+void FPrimitiveComponentVisualizer::DrawVisualization( const UActorComponent* Component, const FSceneView* View, FPrimitiveDrawInterface* PDI )
+{
+	if (const UPrimitiveComponent* PrimComp = Cast<const UPrimitiveComponent>(Component))
+	{
+		if (PrimComp->IsSimulatingPhysics())
+		{
+			if (FBodyInstance* BI = PrimComp->GetBodyInstance())
+			{
+				const FVector& COMPosition = BI->GetCOMPosition();
+				DrawCircle(PDI, COMPosition, FVector(1, 0, 0), FVector(0, 1, 0), COMColor, 15, 25, SDPG_World);
+				DrawCircle(PDI, COMPosition, FVector(0, 0, 1), FVector(0, 1, 0), COMColor, 15, 25, SDPG_World);
+			}
+		}
+	}
+}
diff --git a/Engine/Source/Editor/ComponentVisualizers/Private/PrimitiveComponentVisualizer.h b/Engine/Source/Editor/ComponentVisualizers/Private/PrimitiveComponentVisualizer.h
new file mode 100644
--- /dev/null
+++ b/Engine/Source/Editor/ComponentVisualizers/Private/PrimitiveComponentVisualizer.h
@@ -0,0 +1,13 @@
+// Copyright 1998-2014 Epic Games, Inc. All Rights Reserved.
+
+#pragma once
+
+#include "ComponentVisualizer.h"
+
+class FPrimitiveComponentVisualizer : public FComponentVisualizer
+{
+public:
+	// Begin FComponentVisualizer interface
+	virtual void DrawVisualization(const UActorComponent* Component, const FSceneView* View, FPrimitiveDrawInterface* PDI) override;
+	// End FComponentVisualizer interface
+};
