Added spline component Modify() calls in the various transactions, so that native components are undoable.

diff --git a/Engine/Source/Editor/ComponentVisualizers/Private/SplineComponentVisualizer.cpp b/Engine/Source/Editor/ComponentVisualizers/Private/SplineComponentVisualizer.cpp
--- a/Engine/Source/Editor/ComponentVisualizers/Private/SplineComponentVisualizer.cpp
+++ b/Engine/Source/Editor/ComponentVisualizers/Private/SplineComponentVisualizer.cpp
@@ -329,6 +329,8 @@ bool FSplineComponentVisualizer::HandleInputDelta(FEditorViewportClient* Viewpor
 		SelectedKeyIndex != INDEX_NONE &&
 		SelectedKeyIndex < NumPoints)
 	{
+		SplineComp->Modify();
+
 		if (ViewportClient->IsAltPressed() && bAllowDuplication)
 		{
 			OnDuplicateKey();
@@ -433,6 +435,7 @@ void FSplineComponentVisualizer::OnDuplicateKey()
 	const FScopedTransaction Transaction(LOCTEXT("DuplicateSplinePoint", "Duplicate Spline Point"));
 	USplineComponent* SplineComp = GetEditedSplineComponent();
 
+	SplineComp->Modify();
 	if (AActor* Owner = SplineComp->GetOwner())
 	{
 		Owner->Modify();
@@ -470,6 +473,7 @@ void FSplineComponentVisualizer::OnAddKey()
 	const FScopedTransaction Transaction(LOCTEXT("AddSplinePoint", "Add Spline Point"));
 	USplineComponent* SplineComp = GetEditedSplineComponent();
 
+	SplineComp->Modify();
 	if (AActor* Owner = SplineComp->GetOwner())
 	{
 		Owner->Modify();
@@ -501,6 +505,7 @@ void FSplineComponentVisualizer::OnDeleteKey()
 	const FScopedTransaction Transaction(LOCTEXT("DeleteSplinePoint", "Delete Spline Point"));
 	USplineComponent* SplineComp = GetEditedSplineComponent();
 
+	SplineComp->Modify();
 	if (AActor* Owner = SplineComp->GetOwner())
 	{
 		Owner->Modify();
@@ -549,6 +554,7 @@ void FSplineComponentVisualizer::OnResetToAutomaticTangent(EInterpCurveMode Mode
 	if (SplineComp->SplineInfo.Points[SelectedKeyIndex].InterpMode != Mode)
 	{
 		const FScopedTransaction Transaction(LOCTEXT("ResetToAutomaticTangent", "Reset to Automatic Tangent"));
+		SplineComp->Modify();
 		if (AActor* Owner = SplineComp->GetOwner())
 		{
 			Owner->Modify();
@@ -581,6 +587,7 @@ void FSplineComponentVisualizer::OnSetKeyType(EInterpCurveMode Mode)
 	if (SplineComp->SplineInfo.Points[SelectedKeyIndex].InterpMode != Mode)
 	{
 		const FScopedTransaction Transaction(LOCTEXT("SetSplinePointType", "Set Spline Point Type"));
+		SplineComp->Modify();
 		if (AActor* Owner = SplineComp->GetOwner())
 		{
 			Owner->Modify();
