Added tangent handles to spline visualization.

diff --git a/Engine/Source/Editor/ComponentVisualizers/Private/SplineComponentVisualizer.cpp b/Engine/Source/Editor/ComponentVisualizers/Private/SplineComponentVisualizer.cpp
--- a/Engine/Source/Editor/ComponentVisualizers/Private/SplineComponentVisualizer.cpp
+++ b/Engine/Source/Editor/ComponentVisualizers/Private/SplineComponentVisualizer.cpp
@@ -7,6 +7,7 @@
 IMPLEMENT_HIT_PROXY(HSplineVisProxy, HComponentVisProxy);
 IMPLEMENT_HIT_PROXY(HSplineKeyProxy, HSplineVisProxy);
 IMPLEMENT_HIT_PROXY(HSplineSegmentProxy, HSplineVisProxy);
+IMPLEMENT_HIT_PROXY(HSplineTangentHandleProxy, HSplineVisProxy);
 
 #define LOCTEXT_NAMESPACE "SplineComponentVisualizer"
 
@@ -69,6 +70,7 @@ FSplineComponentVisualizer::FSplineComponentVisualizer()
 	: FComponentVisualizer()
 	, SelectedKeyIndex(INDEX_NONE)
 	, SelectedSegmentIndex(INDEX_NONE)
+	, SelectedTangentHandle(ESelectedTangentHandle::None)
 	, bAllowDuplication(true)
 {
 	FSplineComponentVisualizerCommands::Register();
@@ -143,15 +145,32 @@ void FSplineComponentVisualizer::DrawVisualization(const UActorComponent* Compon
 		const FColor NormalColor = SplineComp->EditorUnselectedSplineSegmentColor;
 		const FColor SelectedColor = SplineComp->EditorSelectedSplineSegmentColor;
 		const float GrabHandleSize = 12.0f;
+		const float TangentHandleSize = 10.0f;
 
 		const FInterpCurveVector& SplineInfo = SplineComp->SplineInfo;
 
+		// Draw the tangent handles before anything else so they will not overdraw the rest of the spline
+		if (SplineComp == EditedSplineComp && SelectedKeyIndex != INDEX_NONE)
+		{
+			const FVector KeyPos = SplineComp->ComponentToWorld.TransformPosition(SplineInfo.Points[SelectedKeyIndex].OutVal);
+			const FVector TangentWorldDirection = SplineComp->ComponentToWorld.TransformVector(SplineInfo.Points[SelectedKeyIndex].LeaveTangent);
+
+			PDI->SetHitProxy(NULL);
+			DrawDashedLine(PDI, KeyPos, KeyPos + TangentWorldDirection, SelectedColor, 5, SDPG_Foreground);
+			DrawDashedLine(PDI, KeyPos, KeyPos - TangentWorldDirection, SelectedColor, 5, SDPG_Foreground);
+			PDI->SetHitProxy(new HSplineTangentHandleProxy(Component, SelectedKeyIndex, false));
+			PDI->DrawPoint(KeyPos + TangentWorldDirection, SelectedColor, TangentHandleSize, SDPG_Foreground);
+			PDI->SetHitProxy(new HSplineTangentHandleProxy(Component, SelectedKeyIndex, true));
+			PDI->DrawPoint(KeyPos - TangentWorldDirection, SelectedColor, TangentHandleSize, SDPG_Foreground);
+			PDI->SetHitProxy(NULL);
+		}
+
 		FVector OldKeyPos(0);
 		float OldKeyTime = 0.f;
 		for (int32 KeyIdx = 0; KeyIdx < SplineInfo.Points.Num(); KeyIdx++)
 		{
-			float NewKeyTime = SplineInfo.Points[KeyIdx].InVal;
-			FVector NewKeyPos = SplineComp->ComponentToWorld.TransformPosition( SplineInfo.Eval(NewKeyTime, FVector::ZeroVector) );
+			const float NewKeyTime = SplineInfo.Points[KeyIdx].InVal;
+			const FVector NewKeyPos = SplineComp->ComponentToWorld.TransformPosition( SplineInfo.Eval(NewKeyTime, FVector::ZeroVector) );
 
 			const FColor KeyColor = (SplineComp == EditedSplineComp && KeyIdx == SelectedKeyIndex) ? SelectedColor : NormalColor;
 
@@ -176,8 +195,8 @@ void FSplineComponentVisualizer::DrawVisualization(const UActorComponent* Compon
 				}
 				else
 				{
-					int32 NumSteps = FMath::CeilToInt((NewKeyTime - OldKeyTime) * 32.0f);
-					float DrawSubstep = (NewKeyTime - OldKeyTime) / NumSteps;
+					const int32 NumSteps = FMath::CeilToInt((NewKeyTime - OldKeyTime) * 32.0f);
+					const float DrawSubstep = (NewKeyTime - OldKeyTime) / NumSteps;
 
 					// Find position on first keyframe.
 					float OldTime = OldKeyTime;
@@ -186,8 +205,8 @@ void FSplineComponentVisualizer::DrawVisualization(const UActorComponent* Compon
 					// Then draw a line for each substep.
 					for (int32 StepIdx = 1; StepIdx<NumSteps + 1; StepIdx++)
 					{
-						float NewTime = OldKeyTime + StepIdx*DrawSubstep;
-						FVector NewPos = SplineComp->ComponentToWorld.TransformPosition( SplineInfo.Eval(NewTime, FVector::ZeroVector) );
+						const float NewTime = OldKeyTime + StepIdx*DrawSubstep;
+						const FVector NewPos = SplineComp->ComponentToWorld.TransformPosition( SplineInfo.Eval(NewTime, FVector::ZeroVector) );
 
 						PDI->DrawLine(OldPos, NewPos, LineColor, SDPG_Foreground);
 
@@ -223,6 +242,7 @@ bool FSplineComponentVisualizer::VisProxyHandleClick(FLevelEditorViewportClient*
 				HSplineKeyProxy* KeyProxy = (HSplineKeyProxy*)VisProxy;
 				SelectedKeyIndex = KeyProxy->KeyIndex;
 				SelectedSegmentIndex = INDEX_NONE;
+				SelectedTangentHandle = ESelectedTangentHandle::None;
 				bEditing = true;
 			}
 			else if (VisProxy->IsA(HSplineSegmentProxy::StaticGetType()))
@@ -234,6 +254,7 @@ bool FSplineComponentVisualizer::VisProxyHandleClick(FLevelEditorViewportClient*
 				HSplineSegmentProxy* SegmentProxy = (HSplineSegmentProxy*)VisProxy;
 				SelectedKeyIndex = SegmentProxy->SegmentIndex;
 				SelectedSegmentIndex = SegmentProxy->SegmentIndex;
+				SelectedTangentHandle = ESelectedTangentHandle::None;
 
 				float SubsegmentStartKey = static_cast<float>(SelectedSegmentIndex);
 				FVector SubsegmentStart = SplineComp->ComponentToWorld.TransformPosition(SplineComp->SplineInfo.Eval(SubsegmentStartKey, FVector::ZeroVector));
@@ -264,6 +285,14 @@ bool FSplineComponentVisualizer::VisProxyHandleClick(FLevelEditorViewportClient*
 				SelectedSplinePosition = BestLocation;
 				bEditing = true;
 			}
+			else if (VisProxy->IsA(HSplineTangentHandleProxy::StaticGetType()))
+			{
+				HSplineTangentHandleProxy* KeyProxy = (HSplineTangentHandleProxy*)VisProxy;
+				SelectedKeyIndex = KeyProxy->KeyIndex;
+				SelectedSegmentIndex = INDEX_NONE;
+				SelectedTangentHandle = KeyProxy->bArriveTangent ? ESelectedTangentHandle::Arrive : ESelectedTangentHandle::Leave;
+				bEditing = true;
+			}
 		}
 		else
 		{
@@ -287,7 +316,23 @@ bool FSplineComponentVisualizer::GetWidgetLocation(const FEditorViewportClient*
 	{
 		if (SelectedKeyIndex < SplineComp->SplineInfo.Points.Num())
 		{
-			OutLocation = SplineComp->ComponentToWorld.TransformPosition( SplineComp->SplineInfo.Points[SelectedKeyIndex].OutVal );
+			const auto& Point = SplineComp->SplineInfo.Points[SelectedKeyIndex];
+
+			switch (SelectedTangentHandle)
+			{
+			case ESelectedTangentHandle::None:
+				OutLocation = SplineComp->ComponentToWorld.TransformPosition(Point.OutVal);
+				break;
+
+			case ESelectedTangentHandle::Leave:
+				OutLocation = SplineComp->ComponentToWorld.TransformPosition(Point.OutVal + Point.LeaveTangent);
+				break;
+
+			case ESelectedTangentHandle::Arrive:
+				OutLocation = SplineComp->ComponentToWorld.TransformPosition(Point.OutVal - Point.ArriveTangent);
+				break;
+			}
+
 			return true;
 		}
 	}
@@ -331,7 +376,7 @@ bool FSplineComponentVisualizer::HandleInputDelta(FEditorViewportClient* Viewpor
 	{
 		SplineComp->Modify();
 
-		if (ViewportClient->IsAltPressed() && bAllowDuplication)
+		if (ViewportClient->IsAltPressed() && SelectedTangentHandle == ESelectedTangentHandle::None && bAllowDuplication)
 		{
 			OnDuplicateKey();
 			// Don't duplicate again until we release LMB
@@ -342,15 +387,26 @@ bool FSplineComponentVisualizer::HandleInputDelta(FEditorViewportClient* Viewpor
 
 		if (!DeltaTranslate.IsZero())
 		{
-			// Find key position in world space
-			FVector CurrentWorldPos = SplineComp->ComponentToWorld.TransformPosition(EditedPoint.OutVal);
-			// Move in world space
-			FVector NewWorldPos = CurrentWorldPos + DeltaTranslate;
-			// Convert back to local space
-			EditedPoint.OutVal = SplineComp->ComponentToWorld.InverseTransformPosition(NewWorldPos);
+			if (SelectedTangentHandle == ESelectedTangentHandle::None)
+			{
+				// Find key position in world space
+				const FVector CurrentWorldPos = SplineComp->ComponentToWorld.TransformPosition(EditedPoint.OutVal);
+				// Move in world space
+				const FVector NewWorldPos = CurrentWorldPos + DeltaTranslate;
+				// Convert back to local space
+				EditedPoint.OutVal = SplineComp->ComponentToWorld.InverseTransformPosition(NewWorldPos);
+			}
+			else
+			{
+				const FVector Delta = (SelectedTangentHandle == ESelectedTangentHandle::Leave) ? DeltaTranslate : -DeltaTranslate;
+				const FVector NewTangent = EditedPoint.LeaveTangent + SplineComp->ComponentToWorld.InverseTransformVector(Delta);
+				EditedPoint.LeaveTangent = NewTangent;
+				EditedPoint.ArriveTangent = NewTangent;
+				EditedPoint.InterpMode = CIM_CurveUser;
+			}
 		}
 
-		if (!DeltaRotate.IsZero())
+		if (!DeltaRotate.IsZero() && SelectedTangentHandle == ESelectedTangentHandle::None)
 		{
 			// Set point tangent as user controlled
 			EditedPoint.InterpMode = CIM_CurveUser;
@@ -360,8 +416,8 @@ bool FSplineComponentVisualizer::HandleInputDelta(FEditorViewportClient* Viewpor
 			EditedPoint.LeaveTangent = NewTangent;
 			EditedPoint.ArriveTangent = NewTangent;
 		}
-		 
-		if (!DeltaScale.IsZero())
+
+		if (!DeltaScale.IsZero() && SelectedTangentHandle == ESelectedTangentHandle::None)
 		{
 			// Set point tangent as user controlled
 			EditedPoint.InterpMode = CIM_CurveUser;
diff --git a/Engine/Source/Editor/ComponentVisualizers/Private/SplineComponentVisualizer.h b/Engine/Source/Editor/ComponentVisualizers/Private/SplineComponentVisualizer.h
--- a/Engine/Source/Editor/ComponentVisualizers/Private/SplineComponentVisualizer.h
+++ b/Engine/Source/Editor/ComponentVisualizers/Private/SplineComponentVisualizer.h
@@ -41,6 +41,21 @@ struct HSplineSegmentProxy : public HSplineVisProxy
 	int32 SegmentIndex;
 };
 
+/** Proxy for a tangent handle */
+struct HSplineTangentHandleProxy : public HSplineVisProxy
+{
+	DECLARE_HIT_PROXY();
+
+	HSplineTangentHandleProxy(const UActorComponent* InComponent, int32 InKeyIndex, bool bInArriveTangent)
+		: HSplineVisProxy(InComponent)
+		, KeyIndex(InKeyIndex)
+		, bArriveTangent(bInArriveTangent)
+	{}
+
+	int32 KeyIndex;
+	bool bArriveTangent;
+};
+
 /** SplineComponent visualizer/edit functionality */
 class FSplineComponentVisualizer : public FComponentVisualizer
 {
@@ -103,6 +118,19 @@ private:
 	/** Index of segment we have selected */
 	int32 SelectedSegmentIndex;
 
+	struct ESelectedTangentHandle
+	{
+		enum Type
+		{
+			None,
+			Leave,
+			Arrive
+		};
+	};
+
+	/** Whether the selected key is a tangent handle, and if so, which type */
+	ESelectedTangentHandle::Type SelectedTangentHandle;
+
 	/** Position on spline we have selected */
 	FVector SelectedSplinePosition;
 
