Added facility to select multiple spline component keys.

diff --git a/Engine/Source/Editor/ComponentVisualizers/Private/SplineComponentVisualizer.cpp b/Engine/Source/Editor/ComponentVisualizers/Private/SplineComponentVisualizer.cpp
--- a/Engine/Source/Editor/ComponentVisualizers/Private/SplineComponentVisualizer.cpp
+++ b/Engine/Source/Editor/ComponentVisualizers/Private/SplineComponentVisualizer.cpp
@@ -68,9 +68,10 @@ public:
 
 FSplineComponentVisualizer::FSplineComponentVisualizer()
 	: FComponentVisualizer()
-	, SelectedKeyIndex(INDEX_NONE)
+	, LastKeyIndexSelected(INDEX_NONE)
 	, SelectedSegmentIndex(INDEX_NONE)
-	, SelectedTangentHandle(ESelectedTangentHandle::None)
+	, SelectedTangentHandle(INDEX_NONE)
+	, SelectedTangentHandleType(ESelectedTangentHandle::None)
 	, bAllowDuplication(true)
 {
 	FSplineComponentVisualizerCommands::Register();
@@ -150,19 +151,25 @@ void FSplineComponentVisualizer::DrawVisualization(const UActorComponent* Compon
 		const FInterpCurveVector& SplineInfo = SplineComp->SplineInfo;
 
 		// Draw the tangent handles before anything else so they will not overdraw the rest of the spline
-		if (SplineComp == EditedSplineComp && SelectedKeyIndex != INDEX_NONE)
+		if (SplineComp == EditedSplineComp)
 		{
-			const FVector KeyPos = SplineComp->ComponentToWorld.TransformPosition(SplineInfo.Points[SelectedKeyIndex].OutVal);
-			const FVector TangentWorldDirection = SplineComp->ComponentToWorld.TransformVector(SplineInfo.Points[SelectedKeyIndex].LeaveTangent);
-
-			PDI->SetHitProxy(NULL);
-			DrawDashedLine(PDI, KeyPos, KeyPos + TangentWorldDirection, SelectedColor, 5, SDPG_Foreground);
-			DrawDashedLine(PDI, KeyPos, KeyPos - TangentWorldDirection, SelectedColor, 5, SDPG_Foreground);
-			PDI->SetHitProxy(new HSplineTangentHandleProxy(Component, SelectedKeyIndex, false));
-			PDI->DrawPoint(KeyPos + TangentWorldDirection, SelectedColor, TangentHandleSize, SDPG_Foreground);
-			PDI->SetHitProxy(new HSplineTangentHandleProxy(Component, SelectedKeyIndex, true));
-			PDI->DrawPoint(KeyPos - TangentWorldDirection, SelectedColor, TangentHandleSize, SDPG_Foreground);
-			PDI->SetHitProxy(NULL);
+			for (int32 SelectedKey : SelectedKeys)
+			{
+				if (SplineInfo.Points[SelectedKey].IsCurveKey())
+				{
+					const FVector KeyPos = SplineComp->ComponentToWorld.TransformPosition(SplineInfo.Points[SelectedKey].OutVal);
+					const FVector TangentWorldDirection = SplineComp->ComponentToWorld.TransformVector(SplineInfo.Points[SelectedKey].LeaveTangent);
+
+					PDI->SetHitProxy(NULL);
+					DrawDashedLine(PDI, KeyPos, KeyPos + TangentWorldDirection, SelectedColor, 5, SDPG_Foreground);
+					DrawDashedLine(PDI, KeyPos, KeyPos - TangentWorldDirection, SelectedColor, 5, SDPG_Foreground);
+					PDI->SetHitProxy(new HSplineTangentHandleProxy(Component, SelectedKey, false));
+					PDI->DrawPoint(KeyPos + TangentWorldDirection, SelectedColor, TangentHandleSize, SDPG_Foreground);
+					PDI->SetHitProxy(new HSplineTangentHandleProxy(Component, SelectedKey, true));
+					PDI->DrawPoint(KeyPos - TangentWorldDirection, SelectedColor, TangentHandleSize, SDPG_Foreground);
+					PDI->SetHitProxy(NULL);
+				}
+			}
 		}
 
 		FVector OldKeyPos(0);
@@ -170,9 +177,9 @@ void FSplineComponentVisualizer::DrawVisualization(const UActorComponent* Compon
 		for (int32 KeyIdx = 0; KeyIdx < SplineInfo.Points.Num(); KeyIdx++)
 		{
 			const float NewKeyTime = SplineInfo.Points[KeyIdx].InVal;
-			const FVector NewKeyPos = SplineComp->ComponentToWorld.TransformPosition( SplineInfo.Eval(NewKeyTime, FVector::ZeroVector) );
+			const FVector NewKeyPos = SplineComp->ComponentToWorld.TransformPosition(SplineInfo.Points[KeyIdx].OutVal);
 
-			const FColor KeyColor = (SplineComp == EditedSplineComp && KeyIdx == SelectedKeyIndex) ? SelectedColor : NormalColor;
+			const FColor KeyColor = (SplineComp == EditedSplineComp && SelectedKeys.Contains(KeyIdx)) ? SelectedColor : NormalColor;
 
 			// Draw the keypoint
 			if (!SplineComp->IsClosedLoop() || KeyIdx < SplineInfo.Points.Num() - 1)
@@ -183,9 +190,9 @@ void FSplineComponentVisualizer::DrawVisualization(const UActorComponent* Compon
 			}
 
 			// If not the first keypoint, draw a line to the previous keypoint.
-			if (KeyIdx>0)
+			if (KeyIdx > 0)
 			{
-				const FColor LineColor = (SplineComp == EditedSplineComp && KeyIdx == SelectedKeyIndex + 1) ? SelectedColor : NormalColor;
+				const FColor LineColor = (SplineComp == EditedSplineComp && SelectedKeys.Contains(KeyIdx - 1)) ? SelectedColor : NormalColor;
 				PDI->SetHitProxy(new HSplineSegmentProxy(Component, KeyIdx - 1));
 
 				// For constant interpolation - don't draw ticks - just draw dotted line.
@@ -224,6 +231,50 @@ void FSplineComponentVisualizer::DrawVisualization(const UActorComponent* Compon
 	}
 }
 
+void FSplineComponentVisualizer::ChangeSelectionState(int32 Index, bool bIsCtrlHeld)
+{
+	if (Index == INDEX_NONE)
+	{
+		SelectedKeys.Empty();
+		LastKeyIndexSelected = INDEX_NONE;
+	}
+	else if (!bIsCtrlHeld)
+	{
+		SelectedKeys.Empty();
+		SelectedKeys.Add(Index);
+		LastKeyIndexSelected = Index;
+	}
+	else
+	{
+		// Add or remove from selection if Ctrl is held
+		if (SelectedKeys.Contains(Index))
+		{
+			// If already in selection, toggle it off
+			SelectedKeys.Remove(Index);
+
+			if (LastKeyIndexSelected == Index)
+			{
+				if (SelectedKeys.Num() == 0)
+				{
+					// Last key selected: clear last key index selected
+					LastKeyIndexSelected = INDEX_NONE;
+				}
+				else
+				{
+					// Arbitarily set last key index selected to first member of the set (so that it is valid)
+					LastKeyIndexSelected = *SelectedKeys.CreateConstIterator();
+				}
+			}
+		}
+		else
+		{
+			// Add to selection
+			SelectedKeys.Add(Index);
+			LastKeyIndexSelected = Index;
+		}
+	}
+}
+
 bool FSplineComponentVisualizer::VisProxyHandleClick(FLevelEditorViewportClient* InViewportClient, HComponentVisProxy* VisProxy, const FViewportClick& Click)
 {
 	bool bEditing = false;
@@ -239,22 +290,34 @@ bool FSplineComponentVisualizer::VisProxyHandleClick(FLevelEditorViewportClient*
 
 			if (VisProxy->IsA(HSplineKeyProxy::StaticGetType()))
 			{
+				// Control point clicked
+
 				HSplineKeyProxy* KeyProxy = (HSplineKeyProxy*)VisProxy;
-				SelectedKeyIndex = KeyProxy->KeyIndex;
+
+				// Modify the selection state, unless right-clicking on an already selected key
+				if (Click.GetKey() != EKeys::RightMouseButton || !SelectedKeys.Contains(KeyProxy->KeyIndex))
+				{
+					ChangeSelectionState(KeyProxy->KeyIndex, InViewportClient->IsCtrlPressed());
+				}
 				SelectedSegmentIndex = INDEX_NONE;
-				SelectedTangentHandle = ESelectedTangentHandle::None;
+				SelectedTangentHandle = INDEX_NONE;
+				SelectedTangentHandleType = ESelectedTangentHandle::None;
+
 				bEditing = true;
 			}
 			else if (VisProxy->IsA(HSplineSegmentProxy::StaticGetType()))
 			{
+				// Spline segment clicked
+
 				// Divide segment into subsegments and test each subsegment against ray representing click position and camera direction.
 				// Closest encounter with the spline determines the spline position.
 				const int32 NumSubdivisions = 16;
 
 				HSplineSegmentProxy* SegmentProxy = (HSplineSegmentProxy*)VisProxy;
-				SelectedKeyIndex = SegmentProxy->SegmentIndex;
+				ChangeSelectionState(SegmentProxy->SegmentIndex, InViewportClient->IsCtrlPressed());
 				SelectedSegmentIndex = SegmentProxy->SegmentIndex;
-				SelectedTangentHandle = ESelectedTangentHandle::None;
+				SelectedTangentHandle = INDEX_NONE;
+				SelectedTangentHandleType = ESelectedTangentHandle::None;
 
 				float SubsegmentStartKey = static_cast<float>(SelectedSegmentIndex);
 				FVector SubsegmentStart = SplineComp->ComponentToWorld.TransformPosition(SplineComp->SplineInfo.Eval(SubsegmentStartKey, FVector::ZeroVector));
@@ -283,14 +346,20 @@ bool FSplineComponentVisualizer::VisProxyHandleClick(FLevelEditorViewportClient*
 				}
 
 				SelectedSplinePosition = BestLocation;
+
 				bEditing = true;
 			}
 			else if (VisProxy->IsA(HSplineTangentHandleProxy::StaticGetType()))
 			{
+				// Tangent handle clicked
+
 				HSplineTangentHandleProxy* KeyProxy = (HSplineTangentHandleProxy*)VisProxy;
-				SelectedKeyIndex = KeyProxy->KeyIndex;
+
+				// Note: don't change key selection when a tangent handle is clicked
 				SelectedSegmentIndex = INDEX_NONE;
-				SelectedTangentHandle = KeyProxy->bArriveTangent ? ESelectedTangentHandle::Arrive : ESelectedTangentHandle::Leave;
+				SelectedTangentHandle = KeyProxy->KeyIndex;
+				SelectedTangentHandleType = KeyProxy->bArriveTangent ? ESelectedTangentHandle::Arrive : ESelectedTangentHandle::Leave;
+
 				bEditing = true;
 			}
 		}
@@ -312,29 +381,35 @@ USplineComponent* FSplineComponentVisualizer::GetEditedSplineComponent() const
 bool FSplineComponentVisualizer::GetWidgetLocation(const FEditorViewportClient* ViewportClient, FVector& OutLocation) const
 {
 	USplineComponent* SplineComp = GetEditedSplineComponent();
-	if (SplineComp && SelectedKeyIndex != INDEX_NONE)
+	if (SplineComp != nullptr)
 	{
-		if (SelectedKeyIndex < SplineComp->SplineInfo.Points.Num())
+		if (SelectedTangentHandle != INDEX_NONE)
 		{
-			const auto& Point = SplineComp->SplineInfo.Points[SelectedKeyIndex];
+			// If tangent handle index is set, use that
+			check(SelectedTangentHandle < SplineComp->SplineInfo.Points.Num());
+			const auto& Point = SplineComp->SplineInfo.Points[SelectedTangentHandle];
 
-			switch (SelectedTangentHandle)
+			check(SelectedTangentHandleType != ESelectedTangentHandle::None);
+			if (SelectedTangentHandleType == ESelectedTangentHandle::Leave)
 			{
-			case ESelectedTangentHandle::None:
-				OutLocation = SplineComp->ComponentToWorld.TransformPosition(Point.OutVal);
-				break;
-
-			case ESelectedTangentHandle::Leave:
 				OutLocation = SplineComp->ComponentToWorld.TransformPosition(Point.OutVal + Point.LeaveTangent);
-				break;
-
-			case ESelectedTangentHandle::Arrive:
+			}
+			else if (SelectedTangentHandleType == ESelectedTangentHandle::Arrive)
+			{
 				OutLocation = SplineComp->ComponentToWorld.TransformPosition(Point.OutVal - Point.ArriveTangent);
-				break;
 			}
 
 			return true;
 		}
+		else if (LastKeyIndexSelected != INDEX_NONE)
+		{
+			// Otherwise use the last key index set
+			check(LastKeyIndexSelected < SplineComp->SplineInfo.Points.Num());
+			check(SelectedKeys.Contains(LastKeyIndexSelected));
+			const auto& Point = SplineComp->SplineInfo.Points[LastKeyIndexSelected];
+			OutLocation = SplineComp->ComponentToWorld.TransformPosition(Point.OutVal);
+			return true;
+		}
 	}
 
 	return false;
@@ -346,11 +421,22 @@ bool FSplineComponentVisualizer::GetCustomInputCoordinateSystem(const FEditorVie
 	if (ViewportClient->GetWidgetCoordSystemSpace() == COORD_Local)
 	{
 		USplineComponent* SplineComp = GetEditedSplineComponent();
-		if (SplineComp && SelectedKeyIndex != INDEX_NONE)
+		if (SplineComp != nullptr)
 		{
-			if (SelectedKeyIndex < SplineComp->SplineInfo.Points.Num())
+			// First look at selected tangent handle for coordinate system
+			int32 Index = SelectedTangentHandle;
+			if (Index == INDEX_NONE)
+			{
+				// If not set, fall back to last key index selected
+				Index = LastKeyIndexSelected;
+			}
+
+			if (Index != INDEX_NONE)
 			{
-				const auto& Point = SplineComp->SplineInfo.Points[SelectedKeyIndex];
+				check(Index < SplineComp->SplineInfo.Points.Num());
+				check(SelectedTangentHandle != INDEX_NONE || SelectedKeys.Contains(Index));
+
+				const auto& Point = SplineComp->SplineInfo.Points[Index];
 				const FVector Tangent = Point.ArriveTangent.IsNearlyZero() ? FVector(1.0f, 0.0f, 0.0f) : Point.ArriveTangent.GetSafeNormal();
 				const FVector Bitangent = (Tangent.Z == 1.0f) ? FVector(1.0f, 0.0f, 0.0f) : FVector(-Tangent.Y, Tangent.X, 0.0f).GetSafeNormal();
 				const FVector Normal = FVector::CrossProduct(Tangent, Bitangent);
@@ -368,91 +454,123 @@ bool FSplineComponentVisualizer::GetCustomInputCoordinateSystem(const FEditorVie
 bool FSplineComponentVisualizer::HandleInputDelta(FEditorViewportClient* ViewportClient, FViewport* Viewport, FVector& DeltaTranslate, FRotator& DeltaRotate, FVector& DeltaScale)
 {
 	USplineComponent* SplineComp = GetEditedSplineComponent();
-	const int32 NumPoints = SplineComp->SplineInfo.Points.Num();
-
-	if (SplineComp && 
-		SelectedKeyIndex != INDEX_NONE &&
-		SelectedKeyIndex < NumPoints)
+	if (SplineComp != nullptr)
 	{
-		SplineComp->Modify();
+		const int32 NumPoints = SplineComp->SplineInfo.Points.Num();
 
-		if (ViewportClient->IsAltPressed() && SelectedTangentHandle == ESelectedTangentHandle::None && bAllowDuplication)
+		bool bChangedFirstPoint = false;
+
+		if (SelectedTangentHandle != INDEX_NONE)
 		{
-			OnDuplicateKey();
-			// Don't duplicate again until we release LMB
-			bAllowDuplication = false;
-		}
+			// When tangent handles are manipulated...
 
-		FInterpCurvePoint<FVector>& EditedPoint = SplineComp->SplineInfo.Points[SelectedKeyIndex];
+			check(SelectedTangentHandle < NumPoints);
 
-		if (!DeltaTranslate.IsZero())
-		{
-			if (SelectedTangentHandle == ESelectedTangentHandle::None)
-			{
-				// Find key position in world space
-				const FVector CurrentWorldPos = SplineComp->ComponentToWorld.TransformPosition(EditedPoint.OutVal);
-				// Move in world space
-				const FVector NewWorldPos = CurrentWorldPos + DeltaTranslate;
-				// Convert back to local space
-				EditedPoint.OutVal = SplineComp->ComponentToWorld.InverseTransformPosition(NewWorldPos);
-			}
-			else
+			FInterpCurvePoint<FVector>& EditedPoint = SplineComp->SplineInfo.Points[SelectedTangentHandle];
+
+			if (!DeltaTranslate.IsZero())
 			{
-				const FVector Delta = (SelectedTangentHandle == ESelectedTangentHandle::Leave) ? DeltaTranslate : -DeltaTranslate;
+				check(SelectedTangentHandleType != ESelectedTangentHandle::None);
+
+				const FVector Delta = (SelectedTangentHandleType == ESelectedTangentHandle::Leave) ? DeltaTranslate : -DeltaTranslate;
 				const FVector NewTangent = EditedPoint.LeaveTangent + SplineComp->ComponentToWorld.InverseTransformVector(Delta);
 				EditedPoint.LeaveTangent = NewTangent;
 				EditedPoint.ArriveTangent = NewTangent;
 				EditedPoint.InterpMode = CIM_CurveUser;
-			}
-		}
-
-		if (!DeltaRotate.IsZero() && SelectedTangentHandle == ESelectedTangentHandle::None)
-		{
-			// Set point tangent as user controlled
-			EditedPoint.InterpMode = CIM_CurveUser;
 
-			// Rotate tangent according to delta rotation
-			const FVector NewTangent = DeltaRotate.RotateVector(EditedPoint.LeaveTangent);
-			EditedPoint.LeaveTangent = NewTangent;
-			EditedPoint.ArriveTangent = NewTangent;
+				if (SelectedTangentHandle == 0)
+				{
+					bChangedFirstPoint = true;
+				}
+			}
 		}
-
-		if (!DeltaScale.IsZero() && SelectedTangentHandle == ESelectedTangentHandle::None)
+		else
 		{
-			// Set point tangent as user controlled
-			EditedPoint.InterpMode = CIM_CurveUser;
+			// When spline keys are manipulated...
 
-			// Break tangent into direction and length so we can change its scale (the 'tension')
-			// independently of its direction.
-			FVector Direction;
-			float Length;
-			EditedPoint.LeaveTangent.ToDirectionAndLength(Direction, Length);
+			check(LastKeyIndexSelected != INDEX_NONE);
+			check(LastKeyIndexSelected < NumPoints);
+			check(SelectedKeys.Num() > 0);
 
-			// Figure out which component has changed, and use it
-			float DeltaScaleValue = (DeltaScale.X != 0.0f) ? DeltaScale.X : ((DeltaScale.Y != 0.0f) ? DeltaScale.Y : DeltaScale.Z);
+			SplineComp->Modify();
 
-			// Change scale, avoiding singularity by never allowing a scale of 0, hence preserving direction.
-			Length += DeltaScaleValue * 10.0f;
-			if (Length == 0.0f)
+			if (ViewportClient->IsAltPressed() && bAllowDuplication)
 			{
-				Length = SMALL_NUMBER;
+				OnDuplicateKey();
+				// Don't duplicate again until we release LMB
+				bAllowDuplication = false;
 			}
 
-			const FVector NewTangent = Direction * Length;
-			EditedPoint.LeaveTangent = NewTangent;
-			EditedPoint.ArriveTangent = NewTangent;
+			for (int32 SelectedKeyIndex : SelectedKeys)
+			{
+				FInterpCurvePoint<FVector>& EditedPoint = SplineComp->SplineInfo.Points[SelectedKeyIndex];
+
+				if (!DeltaTranslate.IsZero())
+				{
+					// Find key position in world space
+					const FVector CurrentWorldPos = SplineComp->ComponentToWorld.TransformPosition(EditedPoint.OutVal);
+					// Move in world space
+					const FVector NewWorldPos = CurrentWorldPos + DeltaTranslate;
+					// Convert back to local space
+					EditedPoint.OutVal = SplineComp->ComponentToWorld.InverseTransformPosition(NewWorldPos);
+				}
+
+				if (!DeltaRotate.IsZero())
+				{
+					// Set point tangent as user controlled
+					EditedPoint.InterpMode = CIM_CurveUser;
+
+					// Rotate tangent according to delta rotation
+					const FVector NewTangent = DeltaRotate.RotateVector(EditedPoint.LeaveTangent);
+					EditedPoint.LeaveTangent = NewTangent;
+					EditedPoint.ArriveTangent = NewTangent;
+				}
+
+				if (!DeltaScale.IsZero())
+				{
+					// Set point tangent as user controlled
+					EditedPoint.InterpMode = CIM_CurveUser;
+
+					// Break tangent into direction and length so we can change its scale (the 'tension')
+					// independently of its direction.
+					FVector Direction;
+					float Length;
+					EditedPoint.LeaveTangent.ToDirectionAndLength(Direction, Length);
+
+					// Figure out which component has changed, and use it
+					float DeltaScaleValue = (DeltaScale.X != 0.0f) ? DeltaScale.X : ((DeltaScale.Y != 0.0f) ? DeltaScale.Y : DeltaScale.Z);
+
+					// Change scale, avoiding singularity by never allowing a scale of 0, hence preserving direction.
+					Length += DeltaScaleValue * 10.0f;
+					if (Length == 0.0f)
+					{
+						Length = SMALL_NUMBER;
+					}
+
+					const FVector NewTangent = Direction * Length;
+					EditedPoint.LeaveTangent = NewTangent;
+					EditedPoint.ArriveTangent = NewTangent;
+				}
+
+				if (SelectedKeyIndex == 0)
+				{
+					bChangedFirstPoint = true;
+				}
+			}
 		}
 
 		// Duplicate first point to last point if it is a closed loop
-		if (SplineComp->IsClosedLoop() && SelectedKeyIndex == 0 && NumPoints > 0)
+		if (SplineComp->IsClosedLoop() && bChangedFirstPoint && NumPoints > 0)
 		{
-			SplineComp->SplineInfo.Points[NumPoints - 1].OutVal = EditedPoint.OutVal;
-			SplineComp->SplineInfo.Points[NumPoints - 1].LeaveTangent = EditedPoint.LeaveTangent;
-			SplineComp->SplineInfo.Points[NumPoints - 1].ArriveTangent = EditedPoint.ArriveTangent;
+			const FInterpCurvePoint<FVector>& EditedPoint = SplineComp->SplineInfo.Points[0];
+			FInterpCurvePoint<FVector>& DuplicatedPoint = SplineComp->SplineInfo.Points[NumPoints - 1];
+			DuplicatedPoint.OutVal = EditedPoint.OutVal;
+			DuplicatedPoint.InterpMode = EditedPoint.InterpMode;
+			DuplicatedPoint.LeaveTangent = EditedPoint.LeaveTangent;
+			DuplicatedPoint.ArriveTangent = EditedPoint.ArriveTangent;
 		}
 
 		NotifyComponentModified();
-
 		return true;
 	}
 
@@ -482,7 +600,10 @@ void FSplineComponentVisualizer::EndEditing()
 {
 	SplineOwningActor = NULL;
 	SplineCompPropName = NAME_None;
-	SelectedKeyIndex = INDEX_NONE;
+	ChangeSelectionState(INDEX_NONE, false);
+	SelectedSegmentIndex = INDEX_NONE;
+	SelectedTangentHandle = INDEX_NONE;
+	SelectedTangentHandleType = ESelectedTangentHandle::None;
 }
 
 
@@ -490,6 +611,10 @@ void FSplineComponentVisualizer::OnDuplicateKey()
 {
 	const FScopedTransaction Transaction(LOCTEXT("DuplicateSplinePoint", "Duplicate Spline Point"));
 	USplineComponent* SplineComp = GetEditedSplineComponent();
+	check(SplineComp != nullptr);
+	check(LastKeyIndexSelected != INDEX_NONE);
+	check(SelectedKeys.Num() > 0);
+	check(SelectedKeys.Contains(LastKeyIndexSelected));
 
 	SplineComp->Modify();
 	if (AActor* Owner = SplineComp->GetOwner())
@@ -497,18 +622,50 @@ void FSplineComponentVisualizer::OnDuplicateKey()
 		Owner->Modify();
 	}
 
+	// Remember whether it's a closed loop for later, and turn off the loop while we manipulate the points
 	const bool bWasLoop = SplineComp->IsClosedLoop();
 	SplineComp->SetClosedLoop(false);
 
-	FInterpCurvePoint<FVector> KeyToDupe = SplineComp->SplineInfo.Points[SelectedKeyIndex];
-	KeyToDupe.InterpMode = CIM_CurveAuto;
-	int32 NewKeyIndex = SplineComp->SplineInfo.Points.Insert(KeyToDupe, SelectedKeyIndex);
-	// move selection to 'next' key
-	SelectedKeyIndex++;
+	// Get a sorted list of all the selected indices, highest to lowest
+	TArray<int32> SelectedKeysSorted;
+	for (int32 SelectedKeyIndex : SelectedKeys)
+	{
+		SelectedKeysSorted.Add(SelectedKeyIndex);
+	}
+	SelectedKeysSorted.Sort([](int32 A, int32 B) { return A > B; });
+
+	// Insert duplicates into the list, highest index first, so that the lower indices remain the same
+	TArray<FInterpCurvePoint<FVector>>& Points = SplineComp->SplineInfo.Points;
+	for (int32 SelectedKeyIndex : SelectedKeysSorted)
+	{
+		FInterpCurvePoint<FVector> KeyToDupe = Points[SelectedKeyIndex];
+		KeyToDupe.InterpMode = CIM_CurveAuto;
+		Points.Insert(KeyToDupe, SelectedKeyIndex);
+	}
+
+	// Repopulate the selected keys
+	SelectedKeys.Empty();
+	int32 Offset = SelectedKeysSorted.Num();
+	for (int32 SelectedKeyIndex : SelectedKeysSorted)
+	{
+		SelectedKeys.Add(SelectedKeyIndex + Offset);
+
+		if (LastKeyIndexSelected == SelectedKeyIndex)
+		{
+			LastKeyIndexSelected += Offset;
+		}
+
+		Offset--;
+	}
+
+	// Unset tangent handle selection
+	SelectedTangentHandle = INDEX_NONE;
+	SelectedTangentHandleType = ESelectedTangentHandle::None;
 
 	// Update Input value for all keys
 	SplineComp->RefreshSplineInputs();
 
+	// Restore old closed loop state
 	SplineComp->SetClosedLoop(bWasLoop);
 
 	NotifyComponentModified();
@@ -528,6 +685,12 @@ void FSplineComponentVisualizer::OnAddKey()
 {
 	const FScopedTransaction Transaction(LOCTEXT("AddSplinePoint", "Add Spline Point"));
 	USplineComponent* SplineComp = GetEditedSplineComponent();
+	check(SplineComp != nullptr);
+	check(LastKeyIndexSelected != INDEX_NONE);
+	check(SelectedKeys.Num() > 0);
+	check(SelectedKeys.Contains(LastKeyIndexSelected));
+	check(SelectedTangentHandle == INDEX_NONE);
+	check(SelectedTangentHandleType == ESelectedTangentHandle::None);
 
 	SplineComp->Modify();
 	if (AActor* Owner = SplineComp->GetOwner())
@@ -535,6 +698,7 @@ void FSplineComponentVisualizer::OnAddKey()
 		Owner->Modify();
 	}
 
+	// Remember whether it's a closed loop for later, and turn off the loop while we manipulate the points
 	const bool bWasLoop = SplineComp->IsClosedLoop();
 	SplineComp->SetClosedLoop(false);
 
@@ -543,13 +707,15 @@ void FSplineComponentVisualizer::OnAddKey()
 	NewKey.OutVal = SplineComp->ComponentToWorld.InverseTransformPosition(SelectedSplinePosition);
 	NewKey.InterpMode = CIM_CurveAuto;
 	int32 NewKeyIndex = SplineComp->SplineInfo.Points.Insert(NewKey, SelectedSegmentIndex + 1);
-	// move selection to 'next' key
-	SelectedKeyIndex = SelectedSegmentIndex + 1;
+
+	// Set selection to 'next' key
+	ChangeSelectionState(SelectedSegmentIndex + 1, false);
 	SelectedSegmentIndex = INDEX_NONE;
 
 	// Update Input value for all keys
 	SplineComp->RefreshSplineInputs();
 
+	// Restore old closed loop state
 	SplineComp->SetClosedLoop(bWasLoop);
 
 	NotifyComponentModified();
@@ -560,6 +726,10 @@ void FSplineComponentVisualizer::OnDeleteKey()
 {
 	const FScopedTransaction Transaction(LOCTEXT("DeleteSplinePoint", "Delete Spline Point"));
 	USplineComponent* SplineComp = GetEditedSplineComponent();
+	check(SplineComp != nullptr);
+	check(LastKeyIndexSelected != INDEX_NONE);
+	check(SelectedKeys.Num() > 0);
+	check(SelectedKeys.Contains(LastKeyIndexSelected));
 
 	SplineComp->Modify();
 	if (AActor* Owner = SplineComp->GetOwner())
@@ -567,16 +737,35 @@ void FSplineComponentVisualizer::OnDeleteKey()
 		Owner->Modify();
 	}
 
+	// Remember whether it's a closed loop for later, and turn off the loop while we manipulate the points
 	const bool bWasLoop = SplineComp->IsClosedLoop();
 	SplineComp->SetClosedLoop(false);
 
-	SplineComp->SplineInfo.Points.RemoveAt(SelectedKeyIndex);
-	SplineComp->RefreshSplineInputs(); // update input value for each key
+	// Get a sorted list of all the selected indices, highest to lowest
+	TArray<int32> SelectedKeysSorted;
+	for (int32 SelectedKeyIndex : SelectedKeys)
+	{
+		SelectedKeysSorted.Add(SelectedKeyIndex);
+	}
+	SelectedKeysSorted.Sort([](int32 A, int32 B) { return A > B; });
 
-	SplineComp->SetClosedLoop(bWasLoop);
+	// Delete selected keys from list, highest index first
+	TArray<FInterpCurvePoint<FVector>>& Points = SplineComp->SplineInfo.Points;
+	for (int32 SelectedKeyIndex : SelectedKeysSorted)
+	{
+		Points.RemoveAt(SelectedKeyIndex);
+	}
 
-	SelectedKeyIndex = INDEX_NONE; // deselect any keys
+	// Unselect everything
+	ChangeSelectionState(INDEX_NONE, false);
 	SelectedSegmentIndex = INDEX_NONE;
+	SelectedTangentHandle = INDEX_NONE;
+	SelectedTangentHandleType = ESelectedTangentHandle::None;
+
+	SplineComp->RefreshSplineInputs(); // update input value for each key
+
+	// Restore old closed loop state
+	SplineComp->SetClosedLoop(bWasLoop);
 
 	NotifyComponentModified();
 }
@@ -585,31 +774,29 @@ void FSplineComponentVisualizer::OnDeleteKey()
 bool FSplineComponentVisualizer::CanDeleteKey() const
 {
 	USplineComponent* SplineComp = GetEditedSplineComponent();
-	return (SplineComp &&
-		SplineComp->SplineInfo.Points.Num() >= 2 &&
-		SelectedSegmentIndex == INDEX_NONE &&
-		SelectedKeyIndex != INDEX_NONE &&
-		SelectedKeyIndex < SplineComp->SplineInfo.Points.Num());
+	return (SplineComp != nullptr &&
+			SelectedKeys.Num() > 0 &&
+			SelectedKeys.Num() != SplineComp->SplineInfo.Points.Num() &&
+			LastKeyIndexSelected != INDEX_NONE);
 }
 
 
 bool FSplineComponentVisualizer::IsKeySelectionValid() const
 {
 	USplineComponent* SplineComp = GetEditedSplineComponent();
-	return (SplineComp &&
-		SelectedSegmentIndex == INDEX_NONE &&
-		SelectedKeyIndex != INDEX_NONE &&
-		SelectedKeyIndex < SplineComp->SplineInfo.Points.Num());
+	return (SplineComp != nullptr &&
+			SelectedKeys.Num() > 0 &&
+			LastKeyIndexSelected != INDEX_NONE);
 }
 
 
 void FSplineComponentVisualizer::OnResetToAutomaticTangent(EInterpCurveMode Mode)
 {
 	USplineComponent* SplineComp = GetEditedSplineComponent();
-
-	if (SplineComp->SplineInfo.Points[SelectedKeyIndex].InterpMode != Mode)
+	if (SplineComp != nullptr)
 	{
 		const FScopedTransaction Transaction(LOCTEXT("ResetToAutomaticTangent", "Reset to Automatic Tangent"));
+
 		SplineComp->Modify();
 		if (AActor* Owner = SplineComp->GetOwner())
 		{
@@ -618,7 +805,16 @@ void FSplineComponentVisualizer::OnResetToAutomaticTangent(EInterpCurveMode Mode
 
 		const bool bWasLoop = SplineComp->IsClosedLoop();
 		SplineComp->SetClosedLoop(false);
-		SplineComp->SplineInfo.Points[SelectedKeyIndex].InterpMode = Mode;
+
+		for (int32 SelectedKeyIndex : SelectedKeys)
+		{
+			auto& Point = SplineComp->SplineInfo.Points[SelectedKeyIndex];
+			if (Point.IsCurveKey())
+			{
+				Point.InterpMode = Mode;
+			}
+		}
+
 		SplineComp->SetClosedLoop(bWasLoop);
 		NotifyComponentModified();
 	}
@@ -628,21 +824,29 @@ void FSplineComponentVisualizer::OnResetToAutomaticTangent(EInterpCurveMode Mode
 bool FSplineComponentVisualizer::CanResetToAutomaticTangent(EInterpCurveMode Mode) const
 {
 	USplineComponent* SplineComp = GetEditedSplineComponent();
-	return (SplineComp &&
-		SelectedKeyIndex != INDEX_NONE &&
-		SelectedKeyIndex < SplineComp->SplineInfo.Points.Num() &&
-		SplineComp->SplineInfo.Points[SelectedKeyIndex].IsCurveKey() &&
-		SplineComp->SplineInfo.Points[SelectedKeyIndex].InterpMode != Mode);
+	if (SplineComp != nullptr && LastKeyIndexSelected != INDEX_NONE)
+	{
+		for (int32 SelectedKeyIndex : SelectedKeys)
+		{
+			const auto& Point = SplineComp->SplineInfo.Points[SelectedKeyIndex];
+			if (Point.IsCurveKey() && Point.InterpMode != Mode)
+			{
+				return true;
+			}
+		}
+	}
+
+	return false;
 }
 
 
 void FSplineComponentVisualizer::OnSetKeyType(EInterpCurveMode Mode)
 {
 	USplineComponent* SplineComp = GetEditedSplineComponent();
-
-	if (SplineComp->SplineInfo.Points[SelectedKeyIndex].InterpMode != Mode)
+	if (SplineComp != nullptr)
 	{
 		const FScopedTransaction Transaction(LOCTEXT("SetSplinePointType", "Set Spline Point Type"));
+
 		SplineComp->Modify();
 		if (AActor* Owner = SplineComp->GetOwner())
 		{
@@ -651,7 +855,12 @@ void FSplineComponentVisualizer::OnSetKeyType(EInterpCurveMode Mode)
 
 		const bool bWasLoop = SplineComp->IsClosedLoop();
 		SplineComp->SetClosedLoop(false);
-		SplineComp->SplineInfo.Points[SelectedKeyIndex].InterpMode = Mode;
+
+		for (int32 SelectedKeyIndex : SelectedKeys)
+		{
+			SplineComp->SplineInfo.Points[SelectedKeyIndex].InterpMode = Mode;
+		}
+
 		SplineComp->SetClosedLoop(bWasLoop);
 		NotifyComponentModified();
 	}
@@ -663,9 +872,16 @@ bool FSplineComponentVisualizer::IsKeyTypeSet(EInterpCurveMode Mode) const
 	if (IsKeySelectionValid())
 	{
 		USplineComponent* SplineComp = GetEditedSplineComponent();
+		check(SplineComp != nullptr);
 
-		const auto& SelectedPoint = SplineComp->SplineInfo.Points[SelectedKeyIndex];
-		return ((Mode == CIM_CurveAuto && SelectedPoint.IsCurveKey()) || SelectedPoint.InterpMode == Mode);
+		for (int32 SelectedKeyIndex : SelectedKeys)
+		{
+			const auto& SelectedPoint = SplineComp->SplineInfo.Points[SelectedKeyIndex];
+			if ((Mode == CIM_CurveAuto && SelectedPoint.IsCurveKey()) || SelectedPoint.InterpMode == Mode)
+			{
+				return true;
+			}
+		}
 	}
 
 	return false;
@@ -682,7 +898,7 @@ TSharedPtr<SWidget> FSplineComponentVisualizer::GenerateContextMenu() const
 			{
 				MenuBuilder.AddMenuEntry(FSplineComponentVisualizerCommands::Get().AddKey);
 			}
-			else if (SelectedKeyIndex != INDEX_NONE)
+			else if (LastKeyIndexSelected != INDEX_NONE)
 			{
 				MenuBuilder.AddMenuEntry(FSplineComponentVisualizerCommands::Get().DeleteKey);
 				MenuBuilder.AddMenuEntry(FSplineComponentVisualizerCommands::Get().DuplicateKey);
@@ -692,16 +908,22 @@ TSharedPtr<SWidget> FSplineComponentVisualizer::GenerateContextMenu() const
 					LOCTEXT("KeyTypeTooltip", "Define the type of the spline point."),
 					FNewMenuDelegate::CreateSP(this, &FSplineComponentVisualizer::GenerateSplinePointTypeSubMenu));
 
-				// Only add the Automatic Tangents submenu if the key is a curve type
+				// Only add the Automatic Tangents submenu if any of the keys is a curve type
 				USplineComponent* SplineComp = GetEditedSplineComponent();
-				if (SplineComp &&
-					SelectedKeyIndex < SplineComp->SplineInfo.Points.Num() &&
-					SplineComp->SplineInfo.Points[SelectedKeyIndex].IsCurveKey())
+				if (SplineComp != nullptr)
 				{
-					MenuBuilder.AddSubMenu(
-						LOCTEXT("ResetToAutomaticTangent", "Reset to Automatic Tangent"),
-						LOCTEXT("ResetToAutomaticTangentTooltip", "Reset the spline point tangent to an automatically generated value."),
-						FNewMenuDelegate::CreateSP(this, &FSplineComponentVisualizer::GenerateTangentTypeSubMenu));
+					for (int32 SelectedKeyIndex : SelectedKeys)
+					{
+						const auto& Point = SplineComp->SplineInfo.Points[SelectedKeyIndex];
+						if (Point.IsCurveKey())
+						{
+							MenuBuilder.AddSubMenu(
+								LOCTEXT("ResetToAutomaticTangent", "Reset to Automatic Tangent"),
+								LOCTEXT("ResetToAutomaticTangentTooltip", "Reset the spline point tangent to an automatically generated value."),
+								FNewMenuDelegate::CreateSP(this, &FSplineComponentVisualizer::GenerateTangentTypeSubMenu));
+							break;
+						}
+					}
 				}
 			}
 		}
diff --git a/Engine/Source/Editor/ComponentVisualizers/Private/SplineComponentVisualizer.h b/Engine/Source/Editor/ComponentVisualizers/Private/SplineComponentVisualizer.h
--- a/Engine/Source/Editor/ComponentVisualizers/Private/SplineComponentVisualizer.h
+++ b/Engine/Source/Editor/ComponentVisualizers/Private/SplineComponentVisualizer.h
@@ -80,6 +80,9 @@ public:
 
 private:
 
+	/** Update the key selection state of the visualizer */
+	void ChangeSelectionState(int32 Index, bool bIsCtrlHeld);
+
 	void NotifyComponentModified();
 
 	void OnDeleteKey();
@@ -112,12 +115,18 @@ private:
 	/** Name of property on the actor that references the spline we are editing */
 	FName SplineCompPropName;
 
-	/** Index of key we have selected */
-	int32 SelectedKeyIndex;
+	/** Index of keys we have selected */
+	TSet<int32> SelectedKeys;
+
+	/** Index of the last key we selected */
+	int32 LastKeyIndexSelected;
 
 	/** Index of segment we have selected */
 	int32 SelectedSegmentIndex;
 
+	/** Index of tangent handle we have selected */
+	int32 SelectedTangentHandle;
+
 	struct ESelectedTangentHandle
 	{
 		enum Type
@@ -128,8 +137,8 @@ private:
 		};
 	};
 
-	/** Whether the selected key is a tangent handle, and if so, which type */
-	ESelectedTangentHandle::Type SelectedTangentHandle;
+	/** The type of the selected tangent handle */
+	ESelectedTangentHandle::Type SelectedTangentHandleType;
 
 	/** Position on spline we have selected */
 	FVector SelectedSplinePosition;
