Added SplineMeshComponent visualizer.

diff --git a/Engine/Source/Editor/ComponentVisualizers/Private/ComponentVisualizers.cpp b/Engine/Source/Editor/ComponentVisualizers/Private/ComponentVisualizers.cpp
--- a/Engine/Source/Editor/ComponentVisualizers/Private/ComponentVisualizers.cpp
+++ b/Engine/Source/Editor/ComponentVisualizers/Private/ComponentVisualizers.cpp
@@ -14,6 +14,7 @@
 #include "ConstraintComponentVisualizer.h"
 #include "SpringArmComponentVisualizer.h"
 #include "SplineComponentVisualizer.h"
+#include "SplineMeshComponentVisualizer.h"
 #include "SensingComponentVisualizer.h"
 #include "SpringComponentVisualizer.h"
 #include "PrimitiveComponentVisualizer.h"
@@ -37,6 +38,7 @@ void FComponentVisualizersModule::StartupModule()
 	RegisterComponentVisualizer(UPhysicsConstraintComponent::StaticClass()->GetFName(), MakeShareable(new FConstraintComponentVisualizer));
 	RegisterComponentVisualizer(USpringArmComponent::StaticClass()->GetFName(), MakeShareable(new FSpringArmComponentVisualizer));
 	RegisterComponentVisualizer(USplineComponent::StaticClass()->GetFName(), MakeShareable(new FSplineComponentVisualizer));
+	RegisterComponentVisualizer(USplineMeshComponent::StaticClass()->GetFName(), MakeShareable(new FSplineMeshComponentVisualizer));
 	RegisterComponentVisualizer(UPawnSensingComponent::StaticClass()->GetFName(), MakeShareable(new FSensingComponentVisualizer));
 	RegisterComponentVisualizer(UPhysicsSpringComponent::StaticClass()->GetFName(), MakeShareable(new FSpringComponentVisualizer));
 	RegisterComponentVisualizer(UPrimitiveComponent::StaticClass()->GetFName(), MakeShareable(new FPrimitiveComponentVisualizer));
diff --git a/Engine/Source/Editor/ComponentVisualizers/Private/SplineMeshComponentVisualizer.cpp b/Engine/Source/Editor/ComponentVisualizers/Private/SplineMeshComponentVisualizer.cpp
new file mode 100644
--- /dev/null
+++ b/Engine/Source/Editor/ComponentVisualizers/Private/SplineMeshComponentVisualizer.cpp
@@ -0,0 +1,385 @@
+// Copyright 1998-2015 Epic Games, Inc. All Rights Reserved.
+
+#include "ComponentVisualizersPrivatePCH.h"
+#include "SplineMeshComponentVisualizer.h"
+#include "ScopedTransaction.h"
+
+IMPLEMENT_HIT_PROXY(HSplineMeshVisProxy, HComponentVisProxy);
+IMPLEMENT_HIT_PROXY(HSplineMeshKeyProxy, HSplineMeshVisProxy);
+IMPLEMENT_HIT_PROXY(HSplineMeshTangentHandleProxy, HSplineMeshVisProxy);
+
+#define LOCTEXT_NAMESPACE "SplineMeshComponentVisualizer"
+
+
+FSplineMeshComponentVisualizer::FSplineMeshComponentVisualizer()
+	: FComponentVisualizer()
+	, SelectedKey(INDEX_NONE)
+	, SelectedTangentHandle(INDEX_NONE)
+	, SelectedTangentHandleType(ESelectedTangentHandle::None)
+{
+}
+
+void FSplineMeshComponentVisualizer::OnRegister()
+{
+}
+
+FSplineMeshComponentVisualizer::~FSplineMeshComponentVisualizer()
+{
+}
+
+void FSplineMeshComponentVisualizer::DrawVisualization(const UActorComponent* Component, const FSceneView* View, FPrimitiveDrawInterface* PDI)
+{
+	if (const USplineMeshComponent* SplineMeshComp = Cast<const USplineMeshComponent>(Component))
+	{
+		if (!SplineMeshComp->bAllowSplineEditingPerInstance)
+		{
+			return;
+		}
+
+		USplineMeshComponent* EditedSplineMeshComp = GetEditedSplineMeshComponent();
+
+		const FColor NormalColor = FColor::Red;
+		const FColor SelectedColor = FColor::White;
+		const FColor Color = (SplineMeshComp == EditedSplineMeshComp) ? SelectedColor : NormalColor;
+		const float GrabHandleSize = 12.0f;
+		const float TangentHandleSize = 10.0f;
+
+		// Create a spline object
+		FInterpCurveVector Spline = GetSpline(SplineMeshComp);
+
+		// Draw the tangent handles before anything else so they will not overdraw the rest of the spline
+		for (int32 PointIndex = 0; PointIndex < 2; PointIndex++)
+		{
+			const FVector KeyPos = SplineMeshComp->ComponentToWorld.TransformPosition(Spline.Points[PointIndex].OutVal);
+			const FVector TangentWorldDirection = SplineMeshComp->ComponentToWorld.TransformVector(Spline.Points[PointIndex].LeaveTangent);
+
+			PDI->SetHitProxy(NULL);
+			DrawDashedLine(PDI, KeyPos, KeyPos + TangentWorldDirection, Color, 5, SDPG_Foreground);
+			DrawDashedLine(PDI, KeyPos, KeyPos - TangentWorldDirection, Color, 5, SDPG_Foreground);
+			PDI->SetHitProxy(new HSplineMeshTangentHandleProxy(Component, PointIndex, false));
+			PDI->DrawPoint(KeyPos + TangentWorldDirection, Color, TangentHandleSize, SDPG_Foreground);
+			PDI->SetHitProxy(new HSplineMeshTangentHandleProxy(Component, PointIndex, true));
+			PDI->DrawPoint(KeyPos - TangentWorldDirection, Color, TangentHandleSize, SDPG_Foreground);
+			PDI->SetHitProxy(NULL);
+		}
+
+		// Draw the keypoints
+		for (int32 PointIndex = 0; PointIndex < 2; PointIndex++)
+		{
+			const FVector NewKeyPos = SplineMeshComp->ComponentToWorld.TransformPosition(Spline.Points[PointIndex].OutVal);
+			PDI->SetHitProxy(new HSplineMeshKeyProxy(Component, PointIndex));
+			PDI->DrawPoint(NewKeyPos, Color, GrabHandleSize, SDPG_Foreground);
+			PDI->SetHitProxy(NULL);
+		}
+
+		// Draw the spline
+		FVector StartPos = SplineMeshComp->ComponentToWorld.TransformPosition(Spline.Points[0].OutVal);
+		for (int32 Step = 1; Step < 32; Step++)
+		{
+			const FVector EndPos = SplineMeshComp->ComponentToWorld.TransformPosition(Spline.Eval(Step / 32.0f, FVector::ZeroVector));
+			PDI->DrawLine(StartPos, EndPos, Color, SDPG_Foreground);
+			StartPos = EndPos;
+		}
+	}
+}
+
+FInterpCurveVector FSplineMeshComponentVisualizer::GetSpline(const USplineMeshComponent* SplineMeshComp) const
+{
+	const FVector StartPosition = SplineMeshComp->GetStartPosition();
+	const FVector EndPosition = SplineMeshComp->GetEndPosition();
+	const FVector StartTangent = SplineMeshComp->GetStartTangent();
+	const FVector EndTangent = SplineMeshComp->GetEndTangent();
+
+	FInterpCurveVector Spline;
+	Spline.Points.Reserve(2);
+	Spline.Points.Emplace(0.0f, StartPosition, StartTangent, StartTangent, CIM_CurveUser);
+	Spline.Points.Emplace(1.0f, EndPosition, EndTangent, EndTangent, CIM_CurveUser);
+
+	return Spline;
+}
+
+
+bool FSplineMeshComponentVisualizer::VisProxyHandleClick(FLevelEditorViewportClient* InViewportClient, HComponentVisProxy* VisProxy, const FViewportClick& Click)
+{
+	bool bEditing = false;
+
+	if (VisProxy && VisProxy->Component.IsValid())
+	{
+		const USplineMeshComponent* SplineMeshComp = CastChecked<const USplineMeshComponent>(VisProxy->Component.Get());
+
+		SplineMeshCompPropName = GetComponentPropertyName(SplineMeshComp);
+		if (SplineMeshCompPropName != NAME_None)
+		{
+			SplineMeshOwningActor = SplineMeshComp->GetOwner();
+
+			if (VisProxy->IsA(HSplineMeshKeyProxy::StaticGetType()))
+			{
+				// Control point clicked
+
+				HSplineMeshKeyProxy* KeyProxy = static_cast<HSplineMeshKeyProxy*>(VisProxy);
+
+				SelectedKey = KeyProxy->KeyIndex;
+				SelectedTangentHandle = INDEX_NONE;
+				SelectedTangentHandleType = ESelectedTangentHandle::None;
+
+				bEditing = true;
+			}
+			else if (VisProxy->IsA(HSplineMeshTangentHandleProxy::StaticGetType()))
+			{
+				// Tangent handle clicked
+
+				HSplineMeshTangentHandleProxy* KeyProxy = static_cast<HSplineMeshTangentHandleProxy*>(VisProxy);
+
+				SelectedKey = INDEX_NONE;
+				SelectedTangentHandle = KeyProxy->KeyIndex;
+				SelectedTangentHandleType = KeyProxy->bArriveTangent ? ESelectedTangentHandle::Arrive : ESelectedTangentHandle::Leave;
+
+				bEditing = true;
+			}
+		}
+		else
+		{
+			SplineMeshOwningActor = NULL;
+		}
+	}
+
+	return bEditing;
+}
+
+USplineMeshComponent* FSplineMeshComponentVisualizer::GetEditedSplineMeshComponent() const
+{
+	return Cast<USplineMeshComponent>(GetComponentFromPropertyName(SplineMeshOwningActor.Get(), SplineMeshCompPropName));
+}
+
+
+bool FSplineMeshComponentVisualizer::GetWidgetLocation(const FEditorViewportClient* ViewportClient, FVector& OutLocation) const
+{
+	USplineMeshComponent* SplineMeshComp = GetEditedSplineMeshComponent();
+	if (SplineMeshComp != nullptr)
+	{
+		FInterpCurveVector Spline = GetSpline(SplineMeshComp);
+
+		if (SelectedTangentHandle != INDEX_NONE)
+		{
+			// If tangent handle index is set, use that
+			check(SelectedTangentHandle < 2);
+			const auto& Point = Spline.Points[SelectedTangentHandle];
+
+			check(SelectedTangentHandleType != ESelectedTangentHandle::None);
+			if (SelectedTangentHandleType == ESelectedTangentHandle::Leave)
+			{
+				OutLocation = SplineMeshComp->ComponentToWorld.TransformPosition(Point.OutVal + Point.LeaveTangent);
+			}
+			else if (SelectedTangentHandleType == ESelectedTangentHandle::Arrive)
+			{
+				OutLocation = SplineMeshComp->ComponentToWorld.TransformPosition(Point.OutVal - Point.ArriveTangent);
+			}
+
+			return true;
+		}
+		else if (SelectedKey != INDEX_NONE)
+		{
+			// Otherwise use the last key index set
+			check(SelectedKey < 2);
+			const auto& Point = Spline.Points[SelectedKey];
+			OutLocation = SplineMeshComp->ComponentToWorld.TransformPosition(Point.OutVal);
+			return true;
+		}
+	}
+
+	return false;
+}
+
+
+bool FSplineMeshComponentVisualizer::GetCustomInputCoordinateSystem(const FEditorViewportClient* ViewportClient, FMatrix& OutMatrix) const
+{
+	if (ViewportClient->GetWidgetCoordSystemSpace() == COORD_Local)
+	{
+		USplineMeshComponent* SplineMeshComp = GetEditedSplineMeshComponent();
+		if (SplineMeshComp != nullptr)
+		{
+			// First look at selected tangent handle for coordinate system
+			int32 Index = SelectedTangentHandle;
+			if (Index == INDEX_NONE)
+			{
+				// If not set, fall back to last key index selected
+				Index = SelectedKey;
+			}
+
+			if (Index != INDEX_NONE)
+			{
+				FInterpCurveVector Spline = GetSpline(SplineMeshComp);
+
+				check(Index < 2);
+				check(SelectedTangentHandle != INDEX_NONE || SelectedKey != INDEX_NONE);
+
+				const auto& Point = Spline.Points[Index];
+				const FVector Tangent = Point.ArriveTangent.IsNearlyZero() ? FVector(1.0f, 0.0f, 0.0f) : Point.ArriveTangent.GetSafeNormal();
+				const FVector Bitangent = (Tangent.Z == 1.0f) ? FVector(1.0f, 0.0f, 0.0f) : FVector(-Tangent.Y, Tangent.X, 0.0f).GetSafeNormal();
+				const FVector Normal = FVector::CrossProduct(Tangent, Bitangent);
+
+				OutMatrix = FMatrix(Tangent, Bitangent, Normal, FVector::ZeroVector) * FQuatRotationTranslationMatrix(SplineMeshComp->ComponentToWorld.GetRotation(), FVector::ZeroVector);
+				return true;
+			}
+		}
+	}
+
+	return false;
+}
+
+
+bool FSplineMeshComponentVisualizer::HandleInputDelta(FEditorViewportClient* ViewportClient, FViewport* Viewport, FVector& DeltaTranslate, FRotator& DeltaRotate, FVector& DeltaScale)
+{
+	USplineMeshComponent* SplineMeshComp = GetEditedSplineMeshComponent();
+	if (SplineMeshComp != nullptr)
+	{
+		if (SelectedTangentHandle != INDEX_NONE)
+		{
+			// When tangent handles are manipulated...
+
+			check(SelectedTangentHandle < 2);
+			const FVector OldTangent = (SelectedTangentHandle == 0) ? SplineMeshComp->GetStartTangent() : SplineMeshComp->GetEndTangent();
+
+			if (!DeltaTranslate.IsZero())
+			{
+				check(SelectedTangentHandleType != ESelectedTangentHandle::None);
+
+				const FVector Delta = (SelectedTangentHandleType == ESelectedTangentHandle::Leave) ? DeltaTranslate : -DeltaTranslate;
+				const FVector NewTangent = OldTangent + SplineMeshComp->ComponentToWorld.InverseTransformVector(Delta);
+
+				SplineMeshComp->Modify();
+
+				if (SelectedTangentHandle == 0)
+				{
+					SplineMeshComp->SetStartTangent(NewTangent);
+				}
+				else
+				{
+					SplineMeshComp->SetEndTangent(NewTangent);
+				}
+			}
+		}
+		else
+		{
+			// When spline keys are manipulated...
+
+			check(SelectedKey != INDEX_NONE);
+			check(SelectedKey < 2);
+
+			SplineMeshComp->Modify();
+
+			FVector KeyPosition = (SelectedKey == 0) ? SplineMeshComp->GetStartPosition() : SplineMeshComp->GetEndPosition();
+			FVector KeyTangent = (SelectedKey == 0) ? SplineMeshComp->GetStartTangent() : SplineMeshComp->GetEndTangent();
+
+			bool bModifiedPosition = false;
+			bool bModifiedTangent = false;
+
+			if (!DeltaTranslate.IsZero())
+			{
+				// Find key position in world space
+				const FVector CurrentWorldPos = SplineMeshComp->ComponentToWorld.TransformPosition(KeyPosition);
+				// Move in world space
+				const FVector NewWorldPos = CurrentWorldPos + DeltaTranslate;
+				// Convert back to local space
+				KeyPosition = SplineMeshComp->ComponentToWorld.InverseTransformPosition(NewWorldPos);
+
+				bModifiedPosition = true;
+			}
+
+			if (!DeltaRotate.IsZero())
+			{
+				// Rotate tangent according to delta rotation
+				KeyTangent = DeltaRotate.RotateVector(KeyTangent);
+
+				bModifiedTangent = true;
+			}
+
+			if (!DeltaScale.IsZero())
+			{
+				// Break tangent into direction and length so we can change its scale (the 'tension')
+				// independently of its direction.
+				FVector Direction;
+				float Length;
+				KeyTangent.ToDirectionAndLength(Direction, Length);
+
+				// Figure out which component has changed, and use it
+				float DeltaScaleValue = (DeltaScale.X != 0.0f) ? DeltaScale.X : ((DeltaScale.Y != 0.0f) ? DeltaScale.Y : DeltaScale.Z);
+
+				// Change scale, avoiding singularity by never allowing a scale of 0, hence preserving direction.
+				Length += DeltaScaleValue * 10.0f;
+				if (Length == 0.0f)
+				{
+					Length = SMALL_NUMBER;
+				}
+
+				KeyTangent = Direction * Length;
+
+				bModifiedTangent = true;
+			}
+
+			if (bModifiedPosition)
+			{
+				if (SelectedKey == 0)
+				{
+					SplineMeshComp->SetStartPosition(KeyPosition);
+				}
+				else
+				{
+					SplineMeshComp->SetEndPosition(KeyPosition);
+				}
+			}
+
+			if (bModifiedTangent)
+			{
+				if (SelectedKey == 0)
+				{
+					SplineMeshComp->SetStartTangent(KeyTangent);
+				}
+				else
+				{
+					SplineMeshComp->SetEndTangent(KeyTangent);
+				}
+			}
+		}
+
+		NotifyComponentModified();
+		return true;
+	}
+
+	return false;
+}
+
+bool FSplineMeshComponentVisualizer::HandleInputKey(FEditorViewportClient* ViewportClient, FViewport* Viewport, FKey Key, EInputEvent Event)
+{
+	return false;
+}
+
+
+void FSplineMeshComponentVisualizer::EndEditing()
+{
+	SplineMeshOwningActor = NULL;
+	SplineMeshCompPropName = NAME_None;
+	SelectedKey = INDEX_NONE;
+	SelectedTangentHandle = INDEX_NONE;
+	SelectedTangentHandleType = ESelectedTangentHandle::None;
+}
+
+
+TSharedPtr<SWidget> FSplineMeshComponentVisualizer::GenerateContextMenu() const
+{
+	return SNullWidget::NullWidget;
+}
+
+
+void FSplineMeshComponentVisualizer::NotifyComponentModified()
+{
+	// Notify of change so any CS is re-run
+	if (SplineMeshOwningActor.IsValid())
+	{
+		SplineMeshOwningActor.Get()->PostEditMove(true);
+	}
+
+	GEditor->RedrawLevelEditingViewports(true);
+}
+
+#undef LOCTEXT_NAMESPACE
diff --git a/Engine/Source/Editor/ComponentVisualizers/Private/SplineMeshComponentVisualizer.h b/Engine/Source/Editor/ComponentVisualizers/Private/SplineMeshComponentVisualizer.h
new file mode 100644
--- /dev/null
+++ b/Engine/Source/Editor/ComponentVisualizers/Private/SplineMeshComponentVisualizer.h
@@ -0,0 +1,100 @@
+// Copyright 1998-2015 Epic Games, Inc. All Rights Reserved.
+
+#pragma once
+
+#include "ComponentVisualizer.h"
+#include "Components/SplineMeshComponent.h"
+
+/** Base class for clickable spline mesh component editing proxies */
+struct HSplineMeshVisProxy : public HComponentVisProxy
+{
+	DECLARE_HIT_PROXY();
+
+	HSplineMeshVisProxy(const UActorComponent* InComponent)
+	: HComponentVisProxy(InComponent, HPP_Wireframe)
+	{}
+};
+
+/** Proxy for a spline mesh component key */
+struct HSplineMeshKeyProxy : public HSplineMeshVisProxy
+{
+	DECLARE_HIT_PROXY();
+
+	HSplineMeshKeyProxy(const UActorComponent* InComponent, int32 InKeyIndex) 
+		: HSplineMeshVisProxy(InComponent)
+		, KeyIndex(InKeyIndex)
+	{}
+
+	int32 KeyIndex;
+};
+
+/** Proxy for a tangent handle */
+struct HSplineMeshTangentHandleProxy : public HSplineMeshVisProxy
+{
+	DECLARE_HIT_PROXY();
+
+	HSplineMeshTangentHandleProxy(const UActorComponent* InComponent, int32 InKeyIndex, bool bInArriveTangent)
+		: HSplineMeshVisProxy(InComponent)
+		, KeyIndex(InKeyIndex)
+		, bArriveTangent(bInArriveTangent)
+	{}
+
+	int32 KeyIndex;
+	bool bArriveTangent;
+};
+
+/** SplineMeshComponent visualizer/edit functionality */
+class FSplineMeshComponentVisualizer : public FComponentVisualizer
+{
+public:
+	FSplineMeshComponentVisualizer();
+	virtual ~FSplineMeshComponentVisualizer();
+
+	// Begin FComponentVisualizer interface
+	virtual void OnRegister() override;
+	virtual void DrawVisualization(const UActorComponent* Component, const FSceneView* View, FPrimitiveDrawInterface* PDI) override;
+	virtual bool VisProxyHandleClick(FLevelEditorViewportClient* InViewportClient, HComponentVisProxy* VisProxy, const FViewportClick& Click) override;
+	virtual void EndEditing() override;
+	virtual bool GetWidgetLocation(const FEditorViewportClient* ViewportClient, FVector& OutLocation) const override;
+	virtual bool GetCustomInputCoordinateSystem(const FEditorViewportClient* ViewportClient, FMatrix& OutMatrix) const override;
+	virtual bool HandleInputDelta(FEditorViewportClient* ViewportClient, FViewport* Viewport, FVector& DeltaTranslate, FRotator& DeltaRotate, FVector& DeltaScale) override;
+	virtual bool HandleInputKey(FEditorViewportClient* ViewportClient, FViewport* Viewport, FKey Key, EInputEvent Event) override;
+	virtual TSharedPtr<SWidget> GenerateContextMenu() const override;
+	// End FComponentVisualizer interface
+
+	/** Get the spline component we are currently editing */
+	USplineMeshComponent* GetEditedSplineMeshComponent() const;
+
+private:
+
+	/** Get a spline object for the specified spline mesh component */
+	FInterpCurveVector GetSpline(const USplineMeshComponent* SplineMeshComp) const;
+
+	/** Syncs changes made by the visualizer in the actual component */
+	void NotifyComponentModified();
+
+	/** Actor that owns the currently edited spline */
+	TWeakObjectPtr<AActor> SplineMeshOwningActor;
+
+	/** Name of property on the actor that references the spline we are editing */
+	FName SplineMeshCompPropName;
+
+	/** Index of the key we selected */
+	int32 SelectedKey;
+
+	/** Index of tangent handle we have selected */
+	int32 SelectedTangentHandle;
+
+	struct ESelectedTangentHandle
+	{
+		enum Type
+		{
+			None,
+			Leave,
+			Arrive
+		};
+	};
+
+	/** The type of the selected tangent handle */
+	ESelectedTangentHandle::Type SelectedTangentHandleType;
+};
diff --git a/Engine/Source/Runtime/Engine/Classes/Components/SplineMeshComponent.h b/Engine/Source/Runtime/Engine/Classes/Components/SplineMeshComponent.h
--- a/Engine/Source/Runtime/Engine/Classes/Components/SplineMeshComponent.h
+++ b/Engine/Source/Runtime/Engine/Classes/Components/SplineMeshComponent.h
@@ -101,6 +101,10 @@ class ENGINE_API USplineMeshComponent : public UStaticMeshComponent, public IInt
 	UPROPERTY(EditAnywhere, Category=SplineMesh)
 	FVector SplineUpDir;
 
+	/** If true, spline keys may be edited per instance in the level viewport. Otherwise, the spline should be initialized in the construction script. */
+	UPROPERTY(EditDefaultsOnly, Category = Spline)
+	uint32 bAllowSplineEditingPerInstance:1;
+
 	/** If true, will use smooth interpolation (ease in/out) for Scale, Roll, and Offset along this section of spline. If false, uses linear */
 	UPROPERTY(EditAnywhere, Category=SplineMesh, AdvancedDisplay)
 	uint32 bSmoothInterpRollScale:1;
@@ -127,6 +131,12 @@ class ENGINE_API USplineMeshComponent : public UStaticMeshComponent, public IInt
 	virtual bool Modify(bool bAlwaysMarkDirty = true) override;
 	//End UObject Interface
 
+	// Begin UActorComponent interface.
+	virtual FComponentInstanceDataBase* GetComponentInstanceData() const override;
+	virtual FName GetComponentInstanceDataType() const override;
+	virtual void ApplyComponentInstanceData(FComponentInstanceDataBase* ComponentInstanceData) override;
+	// End UActorComponent interface.
+
 	//Begin USceneComponent Interface
 	virtual FPrimitiveSceneProxy* CreateSceneProxy() override;
 	virtual FBoxSphereBounds CalcBounds(const FTransform& LocalToWorld) const override;
diff --git a/Engine/Source/Runtime/Engine/Private/SplineMeshComponent.cpp b/Engine/Source/Runtime/Engine/Private/SplineMeshComponent.cpp
--- a/Engine/Source/Runtime/Engine/Private/SplineMeshComponent.cpp
+++ b/Engine/Source/Runtime/Engine/Private/SplineMeshComponent.cpp
@@ -223,6 +223,8 @@ USplineMeshComponent::USplineMeshComponent(const FObjectInitializer& ObjectIniti
 	Mobility = EComponentMobility::Static;
 
 	SetCollisionProfileName(UCollisionProfile::NoCollision_ProfileName);
+	bAllowSplineEditingPerInstance = false;
+	bSmoothInterpRollScale = false;
 	bHasCustomNavigableGeometry = EHasCustomNavigableGeometry::Yes;
 
 	SplineUpDir.Z = 1.0f;
@@ -787,6 +789,58 @@ void USplineMeshComponent::RecreateCollision()
 }
 #endif
 
+/** Used to store spline mesh data during RerunConstructionScripts */
+class FSplineMeshInstanceData : public FComponentInstanceDataBase
+{
+public:
+	explicit FSplineMeshInstanceData(const USplineMeshComponent* SourceComponent)
+		: FComponentInstanceDataBase(SourceComponent)
+	{
+	}
+
+	FVector StartPos;
+	FVector EndPos;
+	FVector StartTangent;
+	FVector EndTangent;
+};
+
+FName USplineMeshComponent::GetComponentInstanceDataType() const
+{
+	static const FName SplineMeshInstanceDataTypeName(TEXT("SplineMeshInstanceData"));
+	return SplineMeshInstanceDataTypeName;
+}
+
+FComponentInstanceDataBase* USplineMeshComponent::GetComponentInstanceData() const
+{
+	FSplineMeshInstanceData* SplineMeshInstanceData = nullptr;
+	if (bAllowSplineEditingPerInstance)
+	{
+		SplineMeshInstanceData = new FSplineMeshInstanceData(this);
+		SplineMeshInstanceData->StartPos = SplineParams.StartPos;
+		SplineMeshInstanceData->EndPos = SplineParams.EndPos;
+		SplineMeshInstanceData->StartTangent = SplineParams.StartTangent;
+		SplineMeshInstanceData->EndTangent = SplineParams.EndTangent;
+	}
+	return SplineMeshInstanceData;
+}
+
+void USplineMeshComponent::ApplyComponentInstanceData(FComponentInstanceDataBase* ComponentInstanceData)
+{
+	if (ComponentInstanceData)
+	{
+		FSplineMeshInstanceData* SplineMeshInstanceData = static_cast<FSplineMeshInstanceData*>(ComponentInstanceData);
+		if (bAllowSplineEditingPerInstance)
+		{
+			SplineParams.StartPos = SplineMeshInstanceData->StartPos;
+			SplineParams.EndPos = SplineMeshInstanceData->EndPos;
+			SplineParams.StartTangent = SplineMeshInstanceData->StartTangent;
+			SplineParams.EndTangent = SplineMeshInstanceData->EndTangent;
+			MarkSplineParamsDirty();
+		}
+	}
+}
+
+
 #include "StaticMeshLight.h"
 /** */
 class FSplineStaticLightingMesh : public FStaticMeshStaticLightingMesh
