Fix clang build error.

diff --git a/Engine/Source/Editor/ComponentVisualizers/Private/SplineComponentVisualizer.cpp b/Engine/Source/Editor/ComponentVisualizers/Private/SplineComponentVisualizer.cpp
--- a/Engine/Source/Editor/ComponentVisualizers/Private/SplineComponentVisualizer.cpp
+++ b/Engine/Source/Editor/ComponentVisualizers/Private/SplineComponentVisualizer.cpp
@@ -161,8 +161,8 @@ void FSplineComponentVisualizer::DrawVisualization(const UActorComponent* Compon
 		const bool bIsSplineEditable = SplineComp->bSplineHasBeenEdited || SplineComp->SplineInfo.Points == Archetype->SplineInfo.Points;
 
 		const FColor ReadOnlyColor = FColor(255, 0, 255, 255);
-		const FColor NormalColor = bIsSplineEditable ? SplineComp->EditorUnselectedSplineSegmentColor : ReadOnlyColor;
-		const FColor SelectedColor = bIsSplineEditable ? SplineComp->EditorSelectedSplineSegmentColor : ReadOnlyColor;
+		const FColor NormalColor = bIsSplineEditable ? FColor(SplineComp->EditorUnselectedSplineSegmentColor) : ReadOnlyColor;
+		const FColor SelectedColor = bIsSplineEditable ? FColor(SplineComp->EditorSelectedSplineSegmentColor) : ReadOnlyColor;
 		const float GrabHandleSize = 12.0f;
 		const float TangentHandleSize = 10.0f;
 
