Redundant reads extracted into local variables to avoid a VC bug in TMultiMap iteration.

diff --git a/Engine/Source/Editor/ComponentVisualizers/Private/AudioComponentVisualizer.cpp b/Engine/Source/Editor/ComponentVisualizers/Private/AudioComponentVisualizer.cpp
--- a/Engine/Source/Editor/ComponentVisualizers/Private/AudioComponentVisualizer.cpp
+++ b/Engine/Source/Editor/ComponentVisualizers/Private/AudioComponentVisualizer.cpp
@@ -19,6 +19,11 @@ void FAudioComponentVisualizer::DrawVisualization( const UActorComponent* Compon
 			TMultiMap<EAttenuationShape::Type, FAttenuationSettings::AttenuationShapeDetails> ShapeDetailsMap;
 			AudioComp->CollectAttenuationShapesForVisualization(ShapeDetailsMap);
 
+			FVector Translation = Transform.GetTranslation();
+			FVector UnitXAxis   = Transform.GetUnitAxis( EAxis::X );
+			FVector UnitYAxis   = Transform.GetUnitAxis( EAxis::Y );
+			FVector UnitZAxis   = Transform.GetUnitAxis( EAxis::Z );
+
 			for ( auto It = ShapeDetailsMap.CreateConstIterator(); It; ++It )
 			{
 				FColor AudioOuterRadiusColor(255, 153, 0);
@@ -31,12 +36,12 @@ void FAudioComponentVisualizer::DrawVisualization( const UActorComponent* Compon
 
 					if (ShapeDetails.Falloff > 0.f)
 					{
-						DrawOrientedWireBox( PDI, Transform.GetTranslation(), Transform.GetUnitAxis( EAxis::X ), Transform.GetUnitAxis( EAxis::Y ), Transform.GetUnitAxis( EAxis::Z ), ShapeDetails.Extents + FVector(ShapeDetails.Falloff), AudioOuterRadiusColor, SDPG_World);
-						DrawOrientedWireBox( PDI, Transform.GetTranslation(), Transform.GetUnitAxis( EAxis::X ), Transform.GetUnitAxis( EAxis::Y ), Transform.GetUnitAxis( EAxis::Z ), ShapeDetails.Extents, AudioInnerRadiusColor, SDPG_World);
+						DrawOrientedWireBox( PDI, Translation, UnitXAxis, UnitYAxis, UnitZAxis, ShapeDetails.Extents + FVector(ShapeDetails.Falloff), AudioOuterRadiusColor, SDPG_World);
+						DrawOrientedWireBox( PDI, Translation, UnitXAxis, UnitYAxis, UnitZAxis, ShapeDetails.Extents, AudioInnerRadiusColor, SDPG_World);
 					}
 					else
 					{
-						DrawOrientedWireBox( PDI, Transform.GetTranslation(), Transform.GetUnitAxis( EAxis::X ), Transform.GetUnitAxis( EAxis::Y ), Transform.GetUnitAxis( EAxis::Z ), ShapeDetails.Extents, AudioOuterRadiusColor, SDPG_World);
+						DrawOrientedWireBox( PDI, Translation, UnitXAxis, UnitYAxis, UnitZAxis, ShapeDetails.Extents, AudioOuterRadiusColor, SDPG_World);
 					}
 					break;
 
@@ -44,12 +49,12 @@ void FAudioComponentVisualizer::DrawVisualization( const UActorComponent* Compon
 
 					if (ShapeDetails.Falloff > 0.f)
 					{
-						DrawWireCapsule( PDI, Transform.GetTranslation(), Transform.GetUnitAxis( EAxis::X ), Transform.GetUnitAxis( EAxis::Y ), Transform.GetUnitAxis( EAxis::Z ), AudioOuterRadiusColor, ShapeDetails.Extents.Y + ShapeDetails.Falloff, ShapeDetails.Extents.X + ShapeDetails.Falloff, 25, SDPG_World);
-						DrawWireCapsule( PDI, Transform.GetTranslation(), Transform.GetUnitAxis( EAxis::X ), Transform.GetUnitAxis( EAxis::Y ), Transform.GetUnitAxis( EAxis::Z ), AudioInnerRadiusColor, ShapeDetails.Extents.Y, ShapeDetails.Extents.X, 25, SDPG_World);
+						DrawWireCapsule( PDI, Translation, UnitXAxis, UnitYAxis, UnitZAxis, AudioOuterRadiusColor, ShapeDetails.Extents.Y + ShapeDetails.Falloff, ShapeDetails.Extents.X + ShapeDetails.Falloff, 25, SDPG_World);
+						DrawWireCapsule( PDI, Translation, UnitXAxis, UnitYAxis, UnitZAxis, AudioInnerRadiusColor, ShapeDetails.Extents.Y, ShapeDetails.Extents.X, 25, SDPG_World);
 					}
 					else
 					{
-						DrawWireCapsule( PDI, Transform.GetTranslation(), Transform.GetUnitAxis( EAxis::X ), Transform.GetUnitAxis( EAxis::Y ), Transform.GetUnitAxis( EAxis::Z ), AudioOuterRadiusColor, ShapeDetails.Extents.Y, ShapeDetails.Extents.X, 25, SDPG_World);
+						DrawWireCapsule( PDI, Translation, UnitXAxis, UnitYAxis, UnitZAxis, AudioOuterRadiusColor, ShapeDetails.Extents.Y, ShapeDetails.Extents.X, 25, SDPG_World);
 					}
 					break;
 
@@ -57,7 +62,7 @@ void FAudioComponentVisualizer::DrawVisualization( const UActorComponent* Compon
 					{
 						FTransform Origin = Transform;
 						Origin.SetScale3D(FVector(1.f));
-						Origin.SetTranslation(Transform.GetTranslation() - (Transform.GetUnitAxis( EAxis::X ) * ShapeDetails.ConeOffset));
+						Origin.SetTranslation(Translation - (UnitXAxis * ShapeDetails.ConeOffset));
 
 						if (ShapeDetails.Falloff > 0.f || ShapeDetails.Extents.Z > 0.f)
 						{
@@ -82,12 +87,12 @@ void FAudioComponentVisualizer::DrawVisualization( const UActorComponent* Compon
 
 					if (ShapeDetails.Falloff > 0.f)
 					{
-						DrawWireSphereAutoSides(PDI, Transform.GetTranslation(), AudioOuterRadiusColor, ShapeDetails.Extents.X + ShapeDetails.Falloff, SDPG_World);
-						DrawWireSphereAutoSides(PDI, Transform.GetTranslation(), AudioInnerRadiusColor, ShapeDetails.Extents.X, SDPG_World);
+						DrawWireSphereAutoSides(PDI, Translation, AudioOuterRadiusColor, ShapeDetails.Extents.X + ShapeDetails.Falloff, SDPG_World);
+						DrawWireSphereAutoSides(PDI, Translation, AudioInnerRadiusColor, ShapeDetails.Extents.X, SDPG_World);
 					}
 					else
 					{
-						DrawWireSphereAutoSides(PDI, Transform.GetTranslation(), AudioOuterRadiusColor, ShapeDetails.Extents.X, SDPG_World);
+						DrawWireSphereAutoSides(PDI, Translation, AudioOuterRadiusColor, ShapeDetails.Extents.X, SDPG_World);
 					}
 					break;
 
