Fixed crash on when enabling Visualize Scale and Roll for a Spline Point

diff --git a/Engine/Source/Editor/ComponentVisualizers/Private/SplineComponentVisualizer.cpp b/Engine/Source/Editor/ComponentVisualizers/Private/SplineComponentVisualizer.cpp
--- a/Engine/Source/Editor/ComponentVisualizers/Private/SplineComponentVisualizer.cpp
+++ b/Engine/Source/Editor/ComponentVisualizers/Private/SplineComponentVisualizer.cpp
@@ -955,8 +955,8 @@ void FSplineComponentVisualizer::OnSetVisualizeRollAndScale()
 bool FSplineComponentVisualizer::IsVisualizingRollAndScale() const
 {
 	USplineComponent* SplineComp = GetEditedSplineComponent();
-	check(SplineComp != nullptr);
-	return SplineComp->bShouldVisualizeScale;
+	
+	return SplineComp ? SplineComp->bShouldVisualizeScale : false;
 }
 
 
