Fixed selection bug in SplineComponent visualizer.  Now correctly selects points on a new spline when a different spline was previously selected.

diff --git a/Engine/Source/Editor/ComponentVisualizers/Private/SplineComponentVisualizer.cpp b/Engine/Source/Editor/ComponentVisualizers/Private/SplineComponentVisualizer.cpp
--- a/Engine/Source/Editor/ComponentVisualizers/Private/SplineComponentVisualizer.cpp
+++ b/Engine/Source/Editor/ComponentVisualizers/Private/SplineComponentVisualizer.cpp
@@ -356,8 +356,18 @@ bool FSplineComponentVisualizer::VisProxyHandleClick(FLevelEditorViewportClient*
 		SplineCompPropName = GetComponentPropertyName(SplineComp);
 		if(SplineCompPropName != NAME_None)
 		{
+			AActor* OldSplineOwningActor = SplineOwningActor.Get();
 			SplineOwningActor = SplineComp->GetOwner();
 
+			if (OldSplineOwningActor != SplineOwningActor)
+			{
+				// Reset selection state if we are selecting a different actor to the one previously selected
+				ChangeSelectionState(INDEX_NONE, false);
+				SelectedSegmentIndex = INDEX_NONE;
+				SelectedTangentHandle = INDEX_NONE;
+				SelectedTangentHandleType = ESelectedTangentHandle::None;
+			}
+
 			if (VisProxy->IsA(HSplineKeyProxy::StaticGetType()))
 			{
 				// Control point clicked
@@ -373,6 +383,12 @@ bool FSplineComponentVisualizer::VisProxyHandleClick(FLevelEditorViewportClient*
 				SelectedTangentHandle = INDEX_NONE;
 				SelectedTangentHandleType = ESelectedTangentHandle::None;
 
+				if (LastKeyIndexSelected == INDEX_NONE)
+				{
+					SplineOwningActor = nullptr;
+					return false;
+				}
+
 				CachedRotation = SplineComp->GetQuaternionAtSplinePoint(LastKeyIndexSelected, ESplineCoordinateSpace::World);
 
 				return true;
@@ -391,6 +407,12 @@ bool FSplineComponentVisualizer::VisProxyHandleClick(FLevelEditorViewportClient*
 				SelectedTangentHandle = INDEX_NONE;
 				SelectedTangentHandleType = ESelectedTangentHandle::None;
 
+				if (LastKeyIndexSelected == INDEX_NONE)
+				{
+					SplineOwningActor = nullptr;
+					return false;
+				}
+
 				CachedRotation = SplineComp->GetQuaternionAtSplinePoint(LastKeyIndexSelected, ESplineCoordinateSpace::World);
 
 				float SubsegmentStartKey = static_cast<float>(SelectedSegmentIndex);
@@ -441,7 +463,7 @@ bool FSplineComponentVisualizer::VisProxyHandleClick(FLevelEditorViewportClient*
 		}
 		else
 		{
-			SplineOwningActor = NULL;
+			SplineOwningActor = nullptr;
 		}
 	}
 
