Moved component visualizer headers to Public folder to allow their use as base classes for project code. Also moved some spline visualizer members from private to protected scope.

diff --git a/Engine/Source/Editor/ComponentVisualizers/Private/AudioComponentVisualizer.h b/Engine/Source/Editor/ComponentVisualizers/Private/AudioComponentVisualizer.h
deleted file mode 100644
diff --git a/Engine/Source/Editor/ComponentVisualizers/Private/ConstraintComponentVisualizer.h b/Engine/Source/Editor/ComponentVisualizers/Private/ConstraintComponentVisualizer.h
deleted file mode 100644
diff --git a/Engine/Source/Editor/ComponentVisualizers/Private/DecalComponentVisualizer.h b/Engine/Source/Editor/ComponentVisualizers/Private/DecalComponentVisualizer.h
deleted file mode 100644
diff --git a/Engine/Source/Editor/ComponentVisualizers/Private/PointLightComponentVisualizer.h b/Engine/Source/Editor/ComponentVisualizers/Private/PointLightComponentVisualizer.h
deleted file mode 100644
diff --git a/Engine/Source/Editor/ComponentVisualizers/Private/PrimitiveComponentVisualizer.h b/Engine/Source/Editor/ComponentVisualizers/Private/PrimitiveComponentVisualizer.h
deleted file mode 100644
diff --git a/Engine/Source/Editor/ComponentVisualizers/Private/RadialForceComponentVisualizer.h b/Engine/Source/Editor/ComponentVisualizers/Private/RadialForceComponentVisualizer.h
deleted file mode 100644
diff --git a/Engine/Source/Editor/ComponentVisualizers/Private/SensingComponentVisualizer.h b/Engine/Source/Editor/ComponentVisualizers/Private/SensingComponentVisualizer.h
deleted file mode 100644
diff --git a/Engine/Source/Editor/ComponentVisualizers/Private/SplineComponentVisualizer.h b/Engine/Source/Editor/ComponentVisualizers/Private/SplineComponentVisualizer.h
deleted file mode 100644
diff --git a/Engine/Source/Editor/ComponentVisualizers/Private/SplineMeshComponentVisualizer.h b/Engine/Source/Editor/ComponentVisualizers/Private/SplineMeshComponentVisualizer.h
deleted file mode 100644
diff --git a/Engine/Source/Editor/ComponentVisualizers/Private/SpotLightComponentVisualizer.h b/Engine/Source/Editor/ComponentVisualizers/Private/SpotLightComponentVisualizer.h
deleted file mode 100644
diff --git a/Engine/Source/Editor/ComponentVisualizers/Private/SpringArmComponentVisualizer.h b/Engine/Source/Editor/ComponentVisualizers/Private/SpringArmComponentVisualizer.h
deleted file mode 100644
diff --git a/Engine/Source/Editor/ComponentVisualizers/Private/SpringComponentVisualizer.h b/Engine/Source/Editor/ComponentVisualizers/Private/SpringComponentVisualizer.h
deleted file mode 100644
diff --git a/Engine/Source/Editor/ComponentVisualizers/Public/AudioComponentVisualizer.h b/Engine/Source/Editor/ComponentVisualizers/Public/AudioComponentVisualizer.h
new file mode 100644
--- /dev/null
+++ b/Engine/Source/Editor/ComponentVisualizers/Public/AudioComponentVisualizer.h
@@ -0,0 +1,13 @@
+// Copyright 1998-2015 Epic Games, Inc. All Rights Reserved.
+
+#pragma once
+
+#include "ComponentVisualizer.h"
+
+class COMPONENTVISUALIZERS_API FAudioComponentVisualizer : public FComponentVisualizer
+{
+public:
+	//~ Begin FComponentVisualizer Interface
+	virtual void DrawVisualization(const UActorComponent* Component, const FSceneView* View, FPrimitiveDrawInterface* PDI) override;
+	//~ End FComponentVisualizer Interface
+};
diff --git a/Engine/Source/Editor/ComponentVisualizers/Public/ConstraintComponentVisualizer.h b/Engine/Source/Editor/ComponentVisualizers/Public/ConstraintComponentVisualizer.h
new file mode 100644
--- /dev/null
+++ b/Engine/Source/Editor/ComponentVisualizers/Public/ConstraintComponentVisualizer.h
@@ -0,0 +1,13 @@
+// Copyright 1998-2015 Epic Games, Inc. All Rights Reserved.
+
+#pragma once
+
+#include "ComponentVisualizer.h"
+
+class COMPONENTVISUALIZERS_API FConstraintComponentVisualizer : public FComponentVisualizer
+{
+public:
+	//~ Begin FComponentVisualizer Interface
+	virtual void DrawVisualization(const UActorComponent* Component, const FSceneView* View, FPrimitiveDrawInterface* PDI) override;
+	//~ End FComponentVisualizer Interface
+};
diff --git a/Engine/Source/Editor/ComponentVisualizers/Public/DecalComponentVisualizer.h b/Engine/Source/Editor/ComponentVisualizers/Public/DecalComponentVisualizer.h
new file mode 100644
--- /dev/null
+++ b/Engine/Source/Editor/ComponentVisualizers/Public/DecalComponentVisualizer.h
@@ -0,0 +1,14 @@
+// Copyright 1998-2015 Epic Games, Inc. All Rights Reserved.
+
+#pragma once
+
+#include "ComponentVisualizer.h"
+#include "Components/DecalComponent.h"
+
+class COMPONENTVISUALIZERS_API FDecalComponentVisualizer : public FComponentVisualizer
+{
+public:
+	//~ Begin FComponentVisualizer Interface
+	virtual void DrawVisualization(const UActorComponent* Component, const FSceneView* View, FPrimitiveDrawInterface* PDI) override;
+	//~ End FComponentVisualizer Interface
+};
diff --git a/Engine/Source/Editor/ComponentVisualizers/Public/PointLightComponentVisualizer.h b/Engine/Source/Editor/ComponentVisualizers/Public/PointLightComponentVisualizer.h
new file mode 100644
--- /dev/null
+++ b/Engine/Source/Editor/ComponentVisualizers/Public/PointLightComponentVisualizer.h
@@ -0,0 +1,13 @@
+// Copyright 1998-2015 Epic Games, Inc. All Rights Reserved.
+
+#pragma once
+
+#include "ComponentVisualizer.h"
+
+class COMPONENTVISUALIZERS_API FPointLightComponentVisualizer : public FComponentVisualizer
+{
+public:
+	//~ Begin FComponentVisualizer Interface
+	virtual void DrawVisualization(const UActorComponent* Component, const FSceneView* View, FPrimitiveDrawInterface* PDI) override;
+	//~ End FComponentVisualizer Interface
+};
diff --git a/Engine/Source/Editor/ComponentVisualizers/Public/PrimitiveComponentVisualizer.h b/Engine/Source/Editor/ComponentVisualizers/Public/PrimitiveComponentVisualizer.h
new file mode 100644
--- /dev/null
+++ b/Engine/Source/Editor/ComponentVisualizers/Public/PrimitiveComponentVisualizer.h
@@ -0,0 +1,13 @@
+// Copyright 1998-2015 Epic Games, Inc. All Rights Reserved.
+
+#pragma once
+
+#include "ComponentVisualizer.h"
+
+class COMPONENTVISUALIZERS_API FPrimitiveComponentVisualizer : public FComponentVisualizer
+{
+public:
+	//~ Begin FComponentVisualizer Interface
+	virtual void DrawVisualization(const UActorComponent* Component, const FSceneView* View, FPrimitiveDrawInterface* PDI) override;
+	//~ End FComponentVisualizer Interface
+};
diff --git a/Engine/Source/Editor/ComponentVisualizers/Public/RadialForceComponentVisualizer.h b/Engine/Source/Editor/ComponentVisualizers/Public/RadialForceComponentVisualizer.h
new file mode 100644
--- /dev/null
+++ b/Engine/Source/Editor/ComponentVisualizers/Public/RadialForceComponentVisualizer.h
@@ -0,0 +1,13 @@
+// Copyright 1998-2015 Epic Games, Inc. All Rights Reserved.
+
+#pragma once
+
+#include "ComponentVisualizer.h"
+
+class COMPONENTVISUALIZERS_API FRadialForceComponentVisualizer : public FComponentVisualizer
+{
+public:
+	//~ Begin FComponentVisualizer Interface
+	virtual void DrawVisualization( const UActorComponent* Component, const FSceneView* View, FPrimitiveDrawInterface* PDI ) override;
+	//~ End FComponentVisualizer Interface
+};
diff --git a/Engine/Source/Editor/ComponentVisualizers/Public/SensingComponentVisualizer.h b/Engine/Source/Editor/ComponentVisualizers/Public/SensingComponentVisualizer.h
new file mode 100644
--- /dev/null
+++ b/Engine/Source/Editor/ComponentVisualizers/Public/SensingComponentVisualizer.h
@@ -0,0 +1,13 @@
+// Copyright 1998-2015 Epic Games, Inc. All Rights Reserved.
+
+#pragma once
+
+#include "ComponentVisualizer.h"
+
+class COMPONENTVISUALIZERS_API FSensingComponentVisualizer : public FComponentVisualizer
+{
+public:
+	//~ Begin FComponentVisualizer Interface
+	virtual void DrawVisualization(const UActorComponent* Component, const FSceneView* View, FPrimitiveDrawInterface* PDI) override;
+	//~ End FComponentVisualizer Interface
+};
diff --git a/Engine/Source/Editor/ComponentVisualizers/Public/SplineComponentVisualizer.h b/Engine/Source/Editor/ComponentVisualizers/Public/SplineComponentVisualizer.h
new file mode 100644
--- /dev/null
+++ b/Engine/Source/Editor/ComponentVisualizers/Public/SplineComponentVisualizer.h
@@ -0,0 +1,157 @@
+// Copyright 1998-2015 Epic Games, Inc. All Rights Reserved.
+
+#pragma once
+
+#include "ComponentVisualizer.h"
+#include "Components/SplineComponent.h"
+
+/** Base class for clickable spline editing proxies */
+struct HSplineVisProxy : public HComponentVisProxy
+{
+	DECLARE_HIT_PROXY();
+
+	HSplineVisProxy(const UActorComponent* InComponent)
+	: HComponentVisProxy(InComponent, HPP_Wireframe)
+	{}
+};
+
+/** Proxy for a spline key */
+struct HSplineKeyProxy : public HSplineVisProxy
+{
+	DECLARE_HIT_PROXY();
+
+	HSplineKeyProxy(const UActorComponent* InComponent, int32 InKeyIndex) 
+		: HSplineVisProxy(InComponent)
+		, KeyIndex(InKeyIndex)
+	{}
+
+	int32 KeyIndex;
+};
+
+/** Proxy for a spline segment */
+struct HSplineSegmentProxy : public HSplineVisProxy
+{
+	DECLARE_HIT_PROXY();
+
+	HSplineSegmentProxy(const UActorComponent* InComponent, int32 InSegmentIndex)
+		: HSplineVisProxy(InComponent)
+		, SegmentIndex(InSegmentIndex)
+	{}
+
+	int32 SegmentIndex;
+};
+
+/** Proxy for a tangent handle */
+struct HSplineTangentHandleProxy : public HSplineVisProxy
+{
+	DECLARE_HIT_PROXY();
+
+	HSplineTangentHandleProxy(const UActorComponent* InComponent, int32 InKeyIndex, bool bInArriveTangent)
+		: HSplineVisProxy(InComponent)
+		, KeyIndex(InKeyIndex)
+		, bArriveTangent(bInArriveTangent)
+	{}
+
+	int32 KeyIndex;
+	bool bArriveTangent;
+};
+
+/** SplineComponent visualizer/edit functionality */
+class COMPONENTVISUALIZERS_API FSplineComponentVisualizer : public FComponentVisualizer
+{
+public:
+	FSplineComponentVisualizer();
+	virtual ~FSplineComponentVisualizer();
+
+	//~ Begin FComponentVisualizer Interface
+	virtual void OnRegister() override;
+	virtual void DrawVisualization(const UActorComponent* Component, const FSceneView* View, FPrimitiveDrawInterface* PDI) override;
+	virtual bool VisProxyHandleClick(FLevelEditorViewportClient* InViewportClient, HComponentVisProxy* VisProxy, const FViewportClick& Click) override;
+	virtual void EndEditing() override;
+	virtual bool GetWidgetLocation(const FEditorViewportClient* ViewportClient, FVector& OutLocation) const override;
+	virtual bool GetCustomInputCoordinateSystem(const FEditorViewportClient* ViewportClient, FMatrix& OutMatrix) const override;
+	virtual bool HandleInputDelta(FEditorViewportClient* ViewportClient, FViewport* Viewport, FVector& DeltaTranslate, FRotator& DeltaRotate, FVector& DeltaScale) override;
+	virtual bool HandleInputKey(FEditorViewportClient* ViewportClient, FViewport* Viewport, FKey Key, EInputEvent Event) override;
+	virtual TSharedPtr<SWidget> GenerateContextMenu() const override;
+	//~ End FComponentVisualizer Interface
+
+	/** Get the spline component we are currently editing */
+	USplineComponent* GetEditedSplineComponent() const;
+
+protected:
+
+	/** Update the key selection state of the visualizer */
+	void ChangeSelectionState(int32 Index, bool bIsCtrlHeld);
+
+	void NotifyComponentModified();
+
+	void OnDeleteKey();
+	bool CanDeleteKey() const;
+
+	void OnDuplicateKey();
+	bool IsKeySelectionValid() const;
+
+	void OnAddKey();
+	bool CanAddKey() const;
+
+	void OnResetToAutomaticTangent(EInterpCurveMode Mode);
+	bool CanResetToAutomaticTangent(EInterpCurveMode Mode) const;
+
+	void OnSetKeyType(EInterpCurveMode Mode);
+	bool IsKeyTypeSet(EInterpCurveMode Mode) const;
+
+	void OnSetVisualizeRollAndScale();
+	bool IsVisualizingRollAndScale() const;
+
+	void OnResetToDefault();
+	bool CanResetToDefault() const;
+
+	/** Generate the submenu containing the available point types */
+	void GenerateSplinePointTypeSubMenu(FMenuBuilder& MenuBuilder) const;
+
+	/** Generate the submenu containing the available auto tangent types */
+	void GenerateTangentTypeSubMenu(FMenuBuilder& MenuBuilder) const;
+
+	/** Output log commands */
+	TSharedPtr<FUICommandList> SplineComponentVisualizerActions;
+
+	/** Actor that owns the currently edited spline */
+	TWeakObjectPtr<AActor> SplineOwningActor;
+
+	/** Name of property on the actor that references the spline we are editing */
+	FPropertyNameAndIndex SplineCompPropName;
+
+	/** Index of keys we have selected */
+	TSet<int32> SelectedKeys;
+
+	/** Index of the last key we selected */
+	int32 LastKeyIndexSelected;
+
+	/** Index of segment we have selected */
+	int32 SelectedSegmentIndex;
+
+	/** Index of tangent handle we have selected */
+	int32 SelectedTangentHandle;
+
+	struct ESelectedTangentHandle
+	{
+		enum Type
+		{
+			None,
+			Leave,
+			Arrive
+		};
+	};
+
+	/** The type of the selected tangent handle */
+	ESelectedTangentHandle::Type SelectedTangentHandleType;
+
+	/** Position on spline we have selected */
+	FVector SelectedSplinePosition;
+
+	/** Cached rotation for this point */
+	FQuat CachedRotation;
+
+	/** Whether we currently allow duplication when dragging */
+	bool bAllowDuplication;
+};
diff --git a/Engine/Source/Editor/ComponentVisualizers/Public/SplineMeshComponentVisualizer.h b/Engine/Source/Editor/ComponentVisualizers/Public/SplineMeshComponentVisualizer.h
new file mode 100644
--- /dev/null
+++ b/Engine/Source/Editor/ComponentVisualizers/Public/SplineMeshComponentVisualizer.h
@@ -0,0 +1,100 @@
+// Copyright 1998-2015 Epic Games, Inc. All Rights Reserved.
+
+#pragma once
+
+#include "ComponentVisualizer.h"
+#include "Components/SplineMeshComponent.h"
+
+/** Base class for clickable spline mesh component editing proxies */
+struct HSplineMeshVisProxy : public HComponentVisProxy
+{
+	DECLARE_HIT_PROXY();
+
+	HSplineMeshVisProxy(const UActorComponent* InComponent)
+	: HComponentVisProxy(InComponent, HPP_Wireframe)
+	{}
+};
+
+/** Proxy for a spline mesh component key */
+struct HSplineMeshKeyProxy : public HSplineMeshVisProxy
+{
+	DECLARE_HIT_PROXY();
+
+	HSplineMeshKeyProxy(const UActorComponent* InComponent, int32 InKeyIndex) 
+		: HSplineMeshVisProxy(InComponent)
+		, KeyIndex(InKeyIndex)
+	{}
+
+	int32 KeyIndex;
+};
+
+/** Proxy for a tangent handle */
+struct HSplineMeshTangentHandleProxy : public HSplineMeshVisProxy
+{
+	DECLARE_HIT_PROXY();
+
+	HSplineMeshTangentHandleProxy(const UActorComponent* InComponent, int32 InKeyIndex, bool bInArriveTangent)
+		: HSplineMeshVisProxy(InComponent)
+		, KeyIndex(InKeyIndex)
+		, bArriveTangent(bInArriveTangent)
+	{}
+
+	int32 KeyIndex;
+	bool bArriveTangent;
+};
+
+/** SplineMeshComponent visualizer/edit functionality */
+class COMPONENTVISUALIZERS_API FSplineMeshComponentVisualizer : public FComponentVisualizer
+{
+public:
+	FSplineMeshComponentVisualizer();
+	virtual ~FSplineMeshComponentVisualizer();
+
+	//~ Begin FComponentVisualizer Interface
+	virtual void OnRegister() override;
+	virtual void DrawVisualization(const UActorComponent* Component, const FSceneView* View, FPrimitiveDrawInterface* PDI) override;
+	virtual bool VisProxyHandleClick(FLevelEditorViewportClient* InViewportClient, HComponentVisProxy* VisProxy, const FViewportClick& Click) override;
+	virtual void EndEditing() override;
+	virtual bool GetWidgetLocation(const FEditorViewportClient* ViewportClient, FVector& OutLocation) const override;
+	virtual bool GetCustomInputCoordinateSystem(const FEditorViewportClient* ViewportClient, FMatrix& OutMatrix) const override;
+	virtual bool HandleInputDelta(FEditorViewportClient* ViewportClient, FViewport* Viewport, FVector& DeltaTranslate, FRotator& DeltaRotate, FVector& DeltaScale) override;
+	virtual bool HandleInputKey(FEditorViewportClient* ViewportClient, FViewport* Viewport, FKey Key, EInputEvent Event) override;
+	virtual TSharedPtr<SWidget> GenerateContextMenu() const override;
+	//~ End FComponentVisualizer Interface
+
+	/** Get the spline component we are currently editing */
+	USplineMeshComponent* GetEditedSplineMeshComponent() const;
+
+protected:
+
+	/** Get a spline object for the specified spline mesh component */
+	FInterpCurveVector GetSpline(const USplineMeshComponent* SplineMeshComp) const;
+
+	/** Syncs changes made by the visualizer in the actual component */
+	void NotifyComponentModified();
+
+	/** Actor that owns the currently edited spline */
+	TWeakObjectPtr<AActor> SplineMeshOwningActor;
+
+	/** Name of property on the actor that references the spline we are editing */
+	FPropertyNameAndIndex SplineMeshCompPropName;
+
+	/** Index of the key we selected */
+	int32 SelectedKey;
+
+	/** Index of tangent handle we have selected */
+	int32 SelectedTangentHandle;
+
+	struct ESelectedTangentHandle
+	{
+		enum Type
+		{
+			None,
+			Leave,
+			Arrive
+		};
+	};
+
+	/** The type of the selected tangent handle */
+	ESelectedTangentHandle::Type SelectedTangentHandleType;
+};
diff --git a/Engine/Source/Editor/ComponentVisualizers/Public/SpotLightComponentVisualizer.h b/Engine/Source/Editor/ComponentVisualizers/Public/SpotLightComponentVisualizer.h
new file mode 100644
--- /dev/null
+++ b/Engine/Source/Editor/ComponentVisualizers/Public/SpotLightComponentVisualizer.h
@@ -0,0 +1,13 @@
+// Copyright 1998-2015 Epic Games, Inc. All Rights Reserved.
+
+#pragma once
+
+#include "ComponentVisualizer.h"
+
+class COMPONENTVISUALIZERS_API FSpotLightComponentVisualizer : public FComponentVisualizer
+{
+public:
+	//~ Begin FComponentVisualizer Interface
+	virtual void DrawVisualization(const UActorComponent* Component, const FSceneView* View, FPrimitiveDrawInterface* PDI) override;
+	//~ End FComponentVisualizer Interface
+};
diff --git a/Engine/Source/Editor/ComponentVisualizers/Public/SpringArmComponentVisualizer.h b/Engine/Source/Editor/ComponentVisualizers/Public/SpringArmComponentVisualizer.h
new file mode 100644
--- /dev/null
+++ b/Engine/Source/Editor/ComponentVisualizers/Public/SpringArmComponentVisualizer.h
@@ -0,0 +1,13 @@
+// Copyright 1998-2015 Epic Games, Inc. All Rights Reserved.
+
+#pragma once
+
+#include "ComponentVisualizer.h"
+
+class COMPONENTVISUALIZERS_API FSpringArmComponentVisualizer : public FComponentVisualizer
+{
+public:
+	//~ Begin FComponentVisualizer Interface
+	virtual void DrawVisualization(const UActorComponent* Component, const FSceneView* View, FPrimitiveDrawInterface* PDI) override;
+	//~ End FComponentVisualizer Interface
+};
diff --git a/Engine/Source/Editor/ComponentVisualizers/Public/SpringComponentVisualizer.h b/Engine/Source/Editor/ComponentVisualizers/Public/SpringComponentVisualizer.h
new file mode 100644
--- /dev/null
+++ b/Engine/Source/Editor/ComponentVisualizers/Public/SpringComponentVisualizer.h
@@ -0,0 +1,13 @@
+// Copyright 1998-2015 Epic Games, Inc. All Rights Reserved.
+
+#pragma once
+
+#include "ComponentVisualizer.h"
+
+class COMPONENTVISUALIZERS_API FSpringComponentVisualizer : public FComponentVisualizer
+{
+public:
+	//~ Begin FComponentVisualizer Interface
+	virtual void DrawVisualization(const UActorComponent* Component, const FSceneView* View, FPrimitiveDrawInterface* PDI) override;
+	//~ End FComponentVisualizer Interface
+};
