Fix non-unity compile issues with attenuated component visualizers

diff --git a/Engine/Source/Editor/ComponentVisualizers/Public/AttenuatedComponentVisualizer.h b/Engine/Source/Editor/ComponentVisualizers/Public/AttenuatedComponentVisualizer.h
--- a/Engine/Source/Editor/ComponentVisualizers/Public/AttenuatedComponentVisualizer.h
+++ b/Engine/Source/Editor/ComponentVisualizers/Public/AttenuatedComponentVisualizer.h
@@ -4,6 +4,8 @@
 
 #include "ComponentVisualizer.h"
 
+struct FEngineShowFlags;
+
 template<class T>
 class COMPONENTVISUALIZERS_API TAttenuatedComponentVisualizer : public FComponentVisualizer
 {
@@ -108,5 +110,5 @@ public:
 	//~ End FComponentVisualizer Interface
 
 private:
-	virtual bool IsVisualizerEnabled(FEngineShowFlags ShowFlags) const = 0;
+	virtual bool IsVisualizerEnabled(const FEngineShowFlags& ShowFlags) const = 0;
 };
diff --git a/Engine/Source/Editor/ComponentVisualizers/Public/AudioComponentVisualizer.h b/Engine/Source/Editor/ComponentVisualizers/Public/AudioComponentVisualizer.h
--- a/Engine/Source/Editor/ComponentVisualizers/Public/AudioComponentVisualizer.h
+++ b/Engine/Source/Editor/ComponentVisualizers/Public/AudioComponentVisualizer.h
@@ -7,7 +7,7 @@
 class COMPONENTVISUALIZERS_API FAudioComponentVisualizer : public TAttenuatedComponentVisualizer<UAudioComponent>
 {
 private:
-	virtual bool IsVisualizerEnabled(FEngineShowFlags ShowFlags) const override
+	virtual bool IsVisualizerEnabled(const FEngineShowFlags& ShowFlags) const override
 	{
 		return ShowFlags.AudioRadius;
 	}
diff --git a/Engine/Source/Editor/ComponentVisualizers/Public/ForceFeedbackComponentVisualizer.h b/Engine/Source/Editor/ComponentVisualizers/Public/ForceFeedbackComponentVisualizer.h
--- a/Engine/Source/Editor/ComponentVisualizers/Public/ForceFeedbackComponentVisualizer.h
+++ b/Engine/Source/Editor/ComponentVisualizers/Public/ForceFeedbackComponentVisualizer.h
@@ -8,7 +8,7 @@
 class COMPONENTVISUALIZERS_API FForceFeedbackComponentVisualizer : public TAttenuatedComponentVisualizer<UForceFeedbackComponent>
 {
 private:
-	virtual bool IsVisualizerEnabled(FEngineShowFlags ShowFlags) const override
+	virtual bool IsVisualizerEnabled(const FEngineShowFlags& ShowFlags) const override
 	{
 		return ShowFlags.ForceFeedbackRadius;
 	}
