Spline component updates for spline control points:
- added Snap/align to marked spline point.
- added Snap/align to floor.
- added Frustum/box select.
- added Focus on selection.

diff --git a/Engine/Source/Editor/ComponentVisualizers/ComponentVisualizers.Build.cs b/Engine/Source/Editor/ComponentVisualizers/ComponentVisualizers.Build.cs
--- a/Engine/Source/Editor/ComponentVisualizers/ComponentVisualizers.Build.cs
+++ b/Engine/Source/Editor/ComponentVisualizers/ComponentVisualizers.Build.cs
@@ -17,6 +17,7 @@ public class ComponentVisualizers : ModuleRules
 				"Slate",
                 "SlateCore",
 				"UnrealEd",
+				"LevelEditor",
                 "PropertyEditor",
                 "EditorStyle",
 				"AIModule",
diff --git a/Engine/Source/Editor/ComponentVisualizers/Private/SplineComponentVisualizer.cpp b/Engine/Source/Editor/ComponentVisualizers/Private/SplineComponentVisualizer.cpp
--- a/Engine/Source/Editor/ComponentVisualizers/Private/SplineComponentVisualizer.cpp
+++ b/Engine/Source/Editor/ComponentVisualizers/Private/SplineComponentVisualizer.cpp
@@ -10,9 +10,12 @@
 #include "UnrealWidget.h"
 #include "Editor.h"
 #include "EditorViewportClient.h"
+#include "EditorViewportCommands.h"
+#include "LevelEditorActions.h"
 #include "Components/SplineComponent.h"
 #include "ScopedTransaction.h"
 #include "ActorEditorUtils.h"
+#include "WorldCollision.h"
 
 IMPLEMENT_HIT_PROXY(HSplineVisProxy, HComponentVisProxy);
 IMPLEMENT_HIT_PROXY(HSplineKeyProxy, HSplineVisProxy);
@@ -47,6 +50,17 @@ public:
 		UI_COMMAND(SetKeyToCurve, "Curve", "Set spline point to Curve type", EUserInterfaceActionType::RadioButton, FInputChord());
 		UI_COMMAND(SetKeyToLinear, "Linear", "Set spline point to Linear type", EUserInterfaceActionType::RadioButton, FInputChord());
 		UI_COMMAND(SetKeyToConstant, "Constant", "Set spline point to Constant type", EUserInterfaceActionType::RadioButton, FInputChord());
+		UI_COMMAND(SnapToMarkedKey, "Snap to Marked Spline Point", "Snap to previously marked spline point.", EUserInterfaceActionType::Button, FInputChord());
+		UI_COMMAND(AlignToMarkedKey, "Align to Marked Spline Point", "Align to previously marked spline point.", EUserInterfaceActionType::Button, FInputChord());
+		UI_COMMAND(SnapAllToSelectedX, "Snap All To Selected X", "Snap all spline points to selected spline point X.", EUserInterfaceActionType::Button, FInputChord());
+		UI_COMMAND(SnapAllToSelectedY, "Snap All To Selected Y", "Snap all spline points to selected spline point Y.", EUserInterfaceActionType::Button, FInputChord());
+		UI_COMMAND(SnapAllToSelectedZ, "Snap All To Selected Z", "Snap all spline points to selected spline point Z.", EUserInterfaceActionType::Button, FInputChord());
+		UI_COMMAND(SetMarkedKey, "Set as Marked Spline Point", "Mark this spline point for snapping.", EUserInterfaceActionType::Button, FInputChord());
+		UI_COMMAND(ResetMarkedKey, "Reset Marked Spline Point", "Clear marked spline point.", EUserInterfaceActionType::Button, FInputChord());
+		UI_COMMAND(SetLockedAxisNone, "None", "New spline point axis is not fixed.", EUserInterfaceActionType::RadioButton, FInputChord());
+		UI_COMMAND(SetLockedAxisX, "X", "Fix X axis when adding new spline points.", EUserInterfaceActionType::RadioButton, FInputChord());
+		UI_COMMAND(SetLockedAxisY, "Y", "Fix Y axis when adding new spline points.", EUserInterfaceActionType::RadioButton, FInputChord());
+		UI_COMMAND(SetLockedAxisZ, "Z", "Fix Z axis when adding new spline points.", EUserInterfaceActionType::RadioButton, FInputChord());
 		UI_COMMAND(VisualizeRollAndScale, "Visualize Roll and Scale", "Whether the visualization should show roll and scale on this spline.", EUserInterfaceActionType::ToggleButton, FInputChord());
 		UI_COMMAND(DiscontinuousSpline, "Allow Discontinuous Splines", "Whether the visualization allows Arrive and Leave tangents to be set separately.", EUserInterfaceActionType::ToggleButton, FInputChord());
 		UI_COMMAND(ResetToDefault, "Reset to Default", "Reset this spline to its archetype default.", EUserInterfaceActionType::Button, FInputChord());
@@ -77,6 +91,39 @@ public:
 	/** Set spline key to Constant type */
 	TSharedPtr<FUICommandInfo> SetKeyToConstant;
 
+	/** Snap to marked key */
+	TSharedPtr<FUICommandInfo> SnapToMarkedKey;
+
+	/** Align to marked key */
+	TSharedPtr<FUICommandInfo> AlignToMarkedKey;
+
+	/** Snap all spline points to selected point X */
+	TSharedPtr<FUICommandInfo> SnapAllToSelectedX;
+
+	/** Snap all spline points to selected point Y */
+	TSharedPtr<FUICommandInfo> SnapAllToSelectedY;
+
+	/** Snap all spline points to selected point Z */
+	TSharedPtr<FUICommandInfo> SnapAllToSelectedZ;
+
+	/** Mark key */
+	TSharedPtr<FUICommandInfo> SetMarkedKey;
+
+	/** Clear marked key */
+	TSharedPtr<FUICommandInfo> ResetMarkedKey;
+
+	/** No axis is locked when adding new spline points */
+	TSharedPtr<FUICommandInfo> SetLockedAxisNone;
+
+	/** Lock X axis when adding new spline points */
+	TSharedPtr<FUICommandInfo> SetLockedAxisX;
+
+	/** Lock Y axis when adding new spline points */
+	TSharedPtr<FUICommandInfo> SetLockedAxisY;
+
+	/** Lock Z axis when adding new spline points */
+	TSharedPtr<FUICommandInfo> SetLockedAxisZ;
+
 	/** Whether the visualization should show roll and scale */
 	TSharedPtr<FUICommandInfo> VisualizeRollAndScale;
 
@@ -87,8 +134,6 @@ public:
 	TSharedPtr<FUICommandInfo> ResetToDefault;
 };
 
-
-
 FSplineComponentVisualizer::FSplineComponentVisualizer()
 	: FComponentVisualizer()
 	, LastKeyIndexSelected(INDEX_NONE)
@@ -96,6 +141,8 @@ FSplineComponentVisualizer::FSplineComponentVisualizer()
 	, SelectedTangentHandle(INDEX_NONE)
 	, SelectedTangentHandleType(ESelectedTangentHandle::None)
 	, bAllowDuplication(true)
+	, MarkedKeyIndex(INDEX_NONE)
+	, AddKeyLockedAxis(EAxis::None)
 {
 	FSplineComponentVisualizerCommands::Register();
 
@@ -151,6 +198,66 @@ void FSplineComponentVisualizer::OnRegister()
 		FCanExecuteAction(),
 		FIsActionChecked::CreateSP(this, &FSplineComponentVisualizer::IsKeyTypeSet, CIM_Constant));
 
+	SplineComponentVisualizerActions->MapAction(
+		Commands.SnapToMarkedKey,
+		FExecuteAction::CreateSP(this, &FSplineComponentVisualizer::OnSnapToMarkedKey, false),
+		FCanExecuteAction::CreateSP(this, &FSplineComponentVisualizer::CanSnapToMarkedKey));
+
+	SplineComponentVisualizerActions->MapAction(
+		Commands.AlignToMarkedKey,
+		FExecuteAction::CreateSP(this, &FSplineComponentVisualizer::OnSnapToMarkedKey, true),
+		FCanExecuteAction::CreateSP(this, &FSplineComponentVisualizer::CanSnapToMarkedKey));
+
+	SplineComponentVisualizerActions->MapAction(
+		Commands.SnapAllToSelectedX,
+		FExecuteAction::CreateSP(this, &FSplineComponentVisualizer::OnSnapAll, EAxis::X),
+		FCanExecuteAction::CreateSP(this, &FSplineComponentVisualizer::CanSnapAll));
+
+	SplineComponentVisualizerActions->MapAction(
+		Commands.SnapAllToSelectedY,
+		FExecuteAction::CreateSP(this, &FSplineComponentVisualizer::OnSnapAll, EAxis::Y),
+		FCanExecuteAction::CreateSP(this, &FSplineComponentVisualizer::CanSnapAll));
+
+	SplineComponentVisualizerActions->MapAction(
+		Commands.SnapAllToSelectedZ,
+		FExecuteAction::CreateSP(this, &FSplineComponentVisualizer::OnSnapAll, EAxis::Z),
+		FCanExecuteAction::CreateSP(this, &FSplineComponentVisualizer::CanSnapAll));
+
+	SplineComponentVisualizerActions->MapAction(
+		Commands.SetMarkedKey,
+		FExecuteAction::CreateSP(this, &FSplineComponentVisualizer::OnSetMarkedKey),
+		FCanExecuteAction::CreateSP(this, &FSplineComponentVisualizer::CanSetMarkedKey));
+
+	SplineComponentVisualizerActions->MapAction(
+		Commands.ResetMarkedKey,
+		FExecuteAction::CreateSP(this, &FSplineComponentVisualizer::OnResetMarkedKey),
+		FCanExecuteAction::CreateSP(this, &FSplineComponentVisualizer::CanResetMarkedKey));
+
+	SplineComponentVisualizerActions->MapAction(
+		Commands.SetLockedAxisNone,
+		FExecuteAction::CreateSP(this, &FSplineComponentVisualizer::OnLockAxis, EAxis::None),
+		FCanExecuteAction(),
+		FIsActionChecked::CreateSP(this, &FSplineComponentVisualizer::IsLockAxisSet, EAxis::None));
+
+	SplineComponentVisualizerActions->MapAction(
+		Commands.SetLockedAxisX,
+		FExecuteAction::CreateSP(this, &FSplineComponentVisualizer::OnLockAxis, EAxis::X),
+		FCanExecuteAction(),
+		FIsActionChecked::CreateSP(this, &FSplineComponentVisualizer::IsLockAxisSet, EAxis::X));
+
+	SplineComponentVisualizerActions->MapAction(
+		Commands.SetLockedAxisY,
+		FExecuteAction::CreateSP(this, &FSplineComponentVisualizer::OnLockAxis, EAxis::Y),
+		FCanExecuteAction(),
+		FIsActionChecked::CreateSP(this, &FSplineComponentVisualizer::IsLockAxisSet, EAxis::Y));
+
+	SplineComponentVisualizerActions->MapAction(
+		Commands.SetLockedAxisZ,
+		FExecuteAction::CreateSP(this, &FSplineComponentVisualizer::OnLockAxis, EAxis::Z),
+		FCanExecuteAction(),
+		FIsActionChecked::CreateSP(this, &FSplineComponentVisualizer::IsLockAxisSet, EAxis::Z));
+
+
 	SplineComponentVisualizerActions->MapAction(
 		Commands.VisualizeRollAndScale,
 		FExecuteAction::CreateSP(this, &FSplineComponentVisualizer::OnSetVisualizeRollAndScale),
@@ -167,6 +274,27 @@ void FSplineComponentVisualizer::OnRegister()
 		Commands.ResetToDefault,
 		FExecuteAction::CreateSP(this, &FSplineComponentVisualizer::OnResetToDefault),
 		FCanExecuteAction::CreateSP(this, &FSplineComponentVisualizer::CanResetToDefault));
+
+	bool bAlign = false;
+	bool bUseLineTrace = false;
+	bool bUseBounds = false;
+	bool bUsePivot = false;
+	SplineComponentVisualizerActions->MapAction(
+		FLevelEditorCommands::Get().SnapToFloor,
+		FExecuteAction::CreateStatic(&FLevelEditorActionCallbacks::SnapToFloor_Clicked, bAlign, bUseLineTrace, bUseBounds, bUsePivot),
+		FCanExecuteAction::CreateStatic(&FLevelEditorActionCallbacks::ActorSelected_CanExecute)
+		);
+
+	bAlign = true;
+	bUseLineTrace = false;
+	bUseBounds = false;
+	bUsePivot = false;
+	SplineComponentVisualizerActions->MapAction(
+		//Commands.AlignToFloor,
+		FLevelEditorCommands::Get().AlignToFloor,
+		FExecuteAction::CreateStatic(&FLevelEditorActionCallbacks::SnapToFloor_Clicked, bAlign, bUseLineTrace, bUseBounds, bUsePivot),
+		FCanExecuteAction::CreateStatic(&FLevelEditorActionCallbacks::ActorSelected_CanExecute)
+		);
 }
 
 FSplineComponentVisualizer::~FSplineComponentVisualizer()
@@ -201,10 +329,12 @@ void FSplineComponentVisualizer::DrawVisualization(const UActorComponent* Compon
 		const FColor ReadOnlyColor = FColor(255, 0, 255, 255);
 		const FColor NormalColor = bIsSplineEditable ? FColor(SplineComp->EditorUnselectedSplineSegmentColor.ToFColor(true)) : ReadOnlyColor;
 		const FColor SelectedColor = bIsSplineEditable ? FColor(SplineComp->EditorSelectedSplineSegmentColor.ToFColor(true)) : ReadOnlyColor;
-		const float GrabHandleSize = 12.0f;
-		const float TangentHandleSize = 10.0f;
+		const FColor MarkedKeyColor = FColor(255, 0, 0, 255);
+		const float GrabHandleSize = 10.0f;
+		const float TangentHandleSize = 8.0f;
 
 		// Draw the tangent handles before anything else so they will not overdraw the rest of the spline
+		// Only draw tangent handle when single key is selected
 		if (SplineComp == EditedSplineComp)
 		{
 			for (int32 SelectedKey : SelectedKeys)
@@ -218,29 +348,20 @@ void FSplineComponentVisualizer::DrawVisualization(const UActorComponent* Compon
 
 					PDI->SetHitProxy(NULL);
 
-					const float DashSize1 = GetDashSize(View, Location, Location + LeaveTangent, 0.01f);
-					if (DashSize1 > 0.0f)
-					{
-						DrawDashedLine(PDI, Location, Location + LeaveTangent, SelectedColor, DashSize1, SDPG_Foreground);
-					}
-
-					const float DashSize2 = GetDashSize(View, Location, Location - ArriveTangent, 0.01f);
-					if (DashSize2 > 0.0f)
-					{
-						DrawDashedLine(PDI, Location, Location - ArriveTangent, SelectedColor, DashSize2, SDPG_Foreground);
-					}
+					PDI->DrawLine(Location, Location + LeaveTangent, NormalColor, SDPG_Foreground);
+					PDI->DrawLine(Location, Location - ArriveTangent, NormalColor, SDPG_Foreground);
 
 					if (bIsSplineEditable)
 					{
 						PDI->SetHitProxy(new HSplineTangentHandleProxy(Component, SelectedKey, false));
 					}
-					PDI->DrawPoint(Location + LeaveTangent, SelectedColor, TangentHandleSize, SDPG_Foreground);
+					PDI->DrawPoint(Location + LeaveTangent, NormalColor, TangentHandleSize, SDPG_Foreground);
 
 					if (bIsSplineEditable)
 					{
 						PDI->SetHitProxy(new HSplineTangentHandleProxy(Component, SelectedKey, true));
 					}
-					PDI->DrawPoint(Location - ArriveTangent, SelectedColor, TangentHandleSize, SDPG_Foreground);
+					PDI->DrawPoint(Location - ArriveTangent, NormalColor, TangentHandleSize, SDPG_Foreground);
 
 					PDI->SetHitProxy(NULL);
 				}
@@ -293,7 +414,7 @@ void FSplineComponentVisualizer::DrawVisualization(const UActorComponent* Compon
 				{
 					PDI->SetHitProxy(new HSplineKeyProxy(Component, KeyIdx));
 				}
-				PDI->DrawPoint(NewKeyPos, KeyColor, GrabHandleSize, SDPG_Foreground);
+				PDI->DrawPoint(NewKeyPos, IsMarkedKey(SplineComp, KeyIdx) ? MarkedKeyColor : KeyColor, GrabHandleSize, SDPG_Foreground);
 				PDI->SetHitProxy(NULL);
 			}
 
@@ -430,6 +551,8 @@ bool FSplineComponentVisualizer::VisProxyHandleClick(FEditorViewportClient* InVi
 			{
 				// Control point clicked
 
+				// temporarily disable
+
 				HSplineKeyProxy* KeyProxy = (HSplineKeyProxy*)VisProxy;
 
 				// Modify the selection state, unless right-clicking on an already selected key
@@ -599,6 +722,27 @@ bool FSplineComponentVisualizer::IsVisualizingArchetype() const
 
 bool FSplineComponentVisualizer::HandleInputDelta(FEditorViewportClient* ViewportClient, FViewport* Viewport, FVector& DeltaTranslate, FRotator& DeltaRotate, FVector& DeltaScale)
 {
+	USplineComponent* SplineComp = GetEditedSplineComponent();
+	if (SplineComp != nullptr)
+	{
+		if (SelectedTangentHandle != INDEX_NONE)
+		{
+			return TransformSelectedTangent(DeltaTranslate);
+		}
+		else
+		{
+			bool bDuplicateKey = ViewportClient->IsAltPressed() && bAllowDuplication;
+			return TransformSelectedKeys(DeltaTranslate, DeltaRotate, DeltaScale, bDuplicateKey);
+		}
+	}
+
+	return false;
+}
+
+bool FSplineComponentVisualizer::TransformSelectedTangent(const FVector& DeltaTranslate)
+{
+	check(SelectedTangentHandle != INDEX_NONE);
+
 	USplineComponent* SplineComp = GetEditedSplineComponent();
 	if (SplineComp != nullptr)
 	{
@@ -608,116 +752,130 @@ bool FSplineComponentVisualizer::HandleInputDelta(FEditorViewportClient* Viewpor
 
 		const int32 NumPoints = SplinePosition.Points.Num();
 
-		if (SelectedTangentHandle != INDEX_NONE)
-		{
-			// When tangent handles are manipulated...
+		check(SelectedTangentHandle < NumPoints);
+		check(SelectedTangentHandleType != ESelectedTangentHandle::None);
 
-			check(SelectedTangentHandle < NumPoints);
+		if (!DeltaTranslate.IsZero())
+		{
+			SplineComp->Modify();
 
-			if (!DeltaTranslate.IsZero())
+			FInterpCurvePoint<FVector>& EditedPoint = SplinePosition.Points[SelectedTangentHandle];
+			if (SplineComp->bAllowDiscontinuousSpline)
 			{
-				check(SelectedTangentHandleType != ESelectedTangentHandle::None);
-
-				SplineComp->Modify();
-
-				FInterpCurvePoint<FVector>& EditedPoint = SplinePosition.Points[SelectedTangentHandle];
-				if (SplineComp->bAllowDiscontinuousSpline)
+				if (SelectedTangentHandleType == ESelectedTangentHandle::Leave)
 				{
-					if (SelectedTangentHandleType == ESelectedTangentHandle::Leave)
-					{
-						EditedPoint.LeaveTangent += SplineComp->GetComponentTransform().InverseTransformVector(DeltaTranslate);
-					}
-					else
-					{
-						EditedPoint.ArriveTangent += SplineComp->GetComponentTransform().InverseTransformVector(-DeltaTranslate);
-					}
+					EditedPoint.LeaveTangent += SplineComp->GetComponentTransform().InverseTransformVector(DeltaTranslate);
 				}
 				else
 				{
-					const FVector Delta = (SelectedTangentHandleType == ESelectedTangentHandle::Leave) ? DeltaTranslate : -DeltaTranslate;
-					const FVector Tangent = EditedPoint.LeaveTangent + SplineComp->GetComponentTransform().InverseTransformVector(Delta);
-
-					EditedPoint.LeaveTangent = Tangent;
-					EditedPoint.ArriveTangent = Tangent;
+					EditedPoint.ArriveTangent += SplineComp->GetComponentTransform().InverseTransformVector(-DeltaTranslate);
 				}
+			}
+			else
+			{
+				const FVector Delta = (SelectedTangentHandleType == ESelectedTangentHandle::Leave) ? DeltaTranslate : -DeltaTranslate;
+				const FVector Tangent = EditedPoint.LeaveTangent + SplineComp->GetComponentTransform().InverseTransformVector(Delta);
 
-				EditedPoint.InterpMode = CIM_CurveUser;
+				EditedPoint.LeaveTangent = Tangent;
+				EditedPoint.ArriveTangent = Tangent;
 			}
+
+			EditedPoint.InterpMode = CIM_CurveUser;
 		}
-		else
+
+		SplineComp->UpdateSpline();
+		SplineComp->bSplineHasBeenEdited = true;
+
+		NotifyPropertyModified(SplineComp, SplineCurvesProperty);
+
+		return true;
+	}
+
+	return false;
+}
+
+bool FSplineComponentVisualizer::TransformSelectedKeys(const FVector& DeltaTranslate, const FRotator& DeltaRotate, const FVector& DeltaScale, bool bDuplicateKey)
+{
+	USplineComponent* SplineComp = GetEditedSplineComponent();
+	if (SplineComp != nullptr)
+	{
+		FInterpCurveVector& SplinePosition = SplineComp->GetSplinePointsPosition();
+		FInterpCurveQuat& SplineRotation = SplineComp->GetSplinePointsRotation();
+		FInterpCurveVector& SplineScale = SplineComp->GetSplinePointsScale();
+
+		const int32 NumPoints = SplinePosition.Points.Num();
+
+		check(LastKeyIndexSelected != INDEX_NONE);
+		check(LastKeyIndexSelected < NumPoints);
+		check(SelectedKeys.Num() > 0);
+
+		SplineComp->Modify();
+
+		if (bDuplicateKey)
 		{
-			// When spline keys are manipulated...
+			DuplicateKey();
 
-			check(LastKeyIndexSelected != INDEX_NONE);
-			check(LastKeyIndexSelected < NumPoints);
-			check(SelectedKeys.Num() > 0);
+			// Don't duplicate again until we release LMB
+			bAllowDuplication = false;
+		}
+
+		for (int32 SelectedKeyIndex : SelectedKeys)
+		{
+			FInterpCurvePoint<FVector>& EditedPoint = SplinePosition.Points[SelectedKeyIndex];
+			FInterpCurvePoint<FQuat>& EditedRotPoint = SplineRotation.Points[SelectedKeyIndex];
+			FInterpCurvePoint<FVector>& EditedScalePoint = SplineScale.Points[SelectedKeyIndex];
 
-			SplineComp->Modify();
 
-			if (ViewportClient->IsAltPressed() && bAllowDuplication)
+			if (!DeltaTranslate.IsZero())
 			{
-				DuplicateKey();
+				// Find key position in world space
+				const FVector CurrentWorldPos = SplineComp->GetComponentTransform().TransformPosition(EditedPoint.OutVal);
+				// Move in world space
+				const FVector NewWorldPos = CurrentWorldPos + DeltaTranslate;
 
-				// Don't duplicate again until we release LMB
-				bAllowDuplication = false;
+				// Convert back to local space
+				EditedPoint.OutVal = SplineComp->GetComponentTransform().InverseTransformPosition(NewWorldPos);
 			}
 
-			for (int32 SelectedKeyIndex : SelectedKeys)
+			if (!DeltaRotate.IsZero())
 			{
-				FInterpCurvePoint<FVector>& EditedPoint = SplinePosition.Points[SelectedKeyIndex];
-				FInterpCurvePoint<FQuat>& EditedRotPoint = SplineRotation.Points[SelectedKeyIndex];
-				FInterpCurvePoint<FVector>& EditedScalePoint = SplineScale.Points[SelectedKeyIndex];
-
-				if (!DeltaTranslate.IsZero())
-				{
-					// Find key position in world space
-					const FVector CurrentWorldPos = SplineComp->GetComponentTransform().TransformPosition(EditedPoint.OutVal);
-					// Move in world space
-					const FVector NewWorldPos = CurrentWorldPos + DeltaTranslate;
-					// Convert back to local space
-					EditedPoint.OutVal = SplineComp->GetComponentTransform().InverseTransformPosition(NewWorldPos);
-				}
+				// Set point tangent as user controlled
+				EditedPoint.InterpMode = CIM_CurveUser;
 
-				if (!DeltaRotate.IsZero())
-				{
-					// Set point tangent as user controlled
-					EditedPoint.InterpMode = CIM_CurveUser;
-
-					// Rotate tangent according to delta rotation
-					FVector NewTangent = SplineComp->GetComponentTransform().GetRotation().RotateVector(EditedPoint.LeaveTangent); // convert local-space tangent vector to world-space
-					NewTangent = DeltaRotate.RotateVector(NewTangent); // apply world-space delta rotation to world-space tangent
-					NewTangent = SplineComp->GetComponentTransform().GetRotation().Inverse().RotateVector(NewTangent); // convert world-space tangent vector back into local-space
-					EditedPoint.LeaveTangent = NewTangent;
-					EditedPoint.ArriveTangent = NewTangent;
-
-					// Rotate spline rotation according to delta rotation
-					FQuat NewRot = SplineComp->GetComponentTransform().GetRotation() * EditedRotPoint.OutVal; // convert local-space rotation to world-space
-					NewRot = DeltaRotate.Quaternion() * NewRot; // apply world-space rotation
-					NewRot = SplineComp->GetComponentTransform().GetRotation().Inverse() * NewRot; // convert world-space rotation to local-space
-					EditedRotPoint.OutVal = NewRot;
-				}
+				// Rotate tangent according to delta rotation
+				FVector NewTangent = SplineComp->GetComponentTransform().GetRotation().RotateVector(EditedPoint.LeaveTangent); // convert local-space tangent vector to world-space
+				NewTangent = DeltaRotate.RotateVector(NewTangent); // apply world-space delta rotation to world-space tangent
+				NewTangent = SplineComp->GetComponentTransform().GetRotation().Inverse().RotateVector(NewTangent); // convert world-space tangent vector back into local-space
+				EditedPoint.LeaveTangent = NewTangent;
+				EditedPoint.ArriveTangent = NewTangent;
+
+				// Rotate spline rotation according to delta rotation
+				FQuat NewRot = SplineComp->GetComponentTransform().GetRotation() * EditedRotPoint.OutVal; // convert local-space rotation to world-space
+				NewRot = DeltaRotate.Quaternion() * NewRot; // apply world-space rotation
+				NewRot = SplineComp->GetComponentTransform().GetRotation().Inverse() * NewRot; // convert world-space rotation to local-space
+				EditedRotPoint.OutVal = NewRot;
+			}
 
-				if (DeltaScale.X != 0.0f)
-				{
-					// Set point tangent as user controlled
-					EditedPoint.InterpMode = CIM_CurveUser;
+			if (DeltaScale.X != 0.0f)
+			{
+				// Set point tangent as user controlled
+				EditedPoint.InterpMode = CIM_CurveUser;
 
-					const FVector NewTangent = EditedPoint.LeaveTangent * (1.0f + DeltaScale.X);
-					EditedPoint.LeaveTangent = NewTangent;
-					EditedPoint.ArriveTangent = NewTangent;
-				}
+				const FVector NewTangent = EditedPoint.LeaveTangent * (1.0f + DeltaScale.X);
+				EditedPoint.LeaveTangent = NewTangent;
+				EditedPoint.ArriveTangent = NewTangent;
+			}
 
-				if (DeltaScale.Y != 0.0f)
-				{
-					// Scale in Y adjusts the scale spline
-					EditedScalePoint.OutVal.Y *= (1.0f + DeltaScale.Y);
-				}
+			if (DeltaScale.Y != 0.0f)
+			{
+				// Scale in Y adjusts the scale spline
+				EditedScalePoint.OutVal.Y *= (1.0f + DeltaScale.Y);
+			}
 
-				if (DeltaScale.Z != 0.0f)
-				{
-					// Scale in Z adjusts the scale spline
-					EditedScalePoint.OutVal.Z *= (1.0f + DeltaScale.Z);
-				}
+			if (DeltaScale.Z != 0.0f)
+			{
+				// Scale in Z adjusts the scale spline
+				EditedScalePoint.OutVal.Z *= (1.0f + DeltaScale.Z);
 			}
 		}
 
@@ -725,6 +883,14 @@ bool FSplineComponentVisualizer::HandleInputDelta(FEditorViewportClient* Viewpor
 		SplineComp->bSplineHasBeenEdited = true;
 
 		NotifyPropertyModified(SplineComp, SplineCurvesProperty);
+
+		if (!DeltaRotate.IsZero())
+		{
+			CachedRotation = SplineComp->GetQuaternionAtSplinePoint(LastKeyIndexSelected, ESplineCoordinateSpace::World);
+		}
+
+		GEditor->RedrawLevelEditingViewports(true);
+
 		return true;
 	}
 
@@ -763,6 +929,412 @@ bool FSplineComponentVisualizer::HandleInputKey(FEditorViewportClient* ViewportC
 	return bHandled;
 }
 
+bool FSplineComponentVisualizer::HandleModifiedClick(FEditorViewportClient* InViewportClient, HHitProxy* HitProxy, const FViewportClick& Click)
+{
+	/*
+	if (Click.IsControlDown())
+	{
+		// Add points on Ctrl-Click if the last spline point is selected.
+
+		USplineComponent* SplineComp = GetEditedSplineComponent();
+		if (SplineComp != nullptr)
+		{
+			FInterpCurveVector& SplinePosition = SplineComp->GetSplinePointsPosition();
+			int32 NumPoints = SplinePosition.Points.Num();
+
+			// to do add end point
+			if (SelectedKeys.Num() == 1 && !SplineComp->IsClosedLoop())
+			{
+				check(LastKeyIndexSelected != INDEX_NONE);
+				check(SelectedKeys.Contains(LastKeyIndexSelected));
+
+				if (LastKeyIndexSelected == 0)
+				{
+					int32 KeyIdx = LastKeyIndexSelected;
+
+					FInterpCurvePoint<FVector>& EditedPoint = SplinePosition.Points[LastKeyIndexSelected];
+
+					FHitResult Hit(1.0f);
+					FCollisionQueryParams Params(SCENE_QUERY_STAT(MoveSplineKeyToTrace), true);
+
+					// Find key position in world space
+					const FVector CurrentWorldPos = SplineComp->GetComponentTransform().TransformPosition(EditedPoint.OutVal);
+
+					FVector DeltaTranslate = FVector::ZeroVector;
+
+					if (SplineComp->GetWorld()->LineTraceSingleByChannel(Hit, Click.GetOrigin(), Click.GetOrigin() + Click.GetDirection() * WORLD_MAX, ECC_WorldStatic, Params))
+					{
+						DeltaTranslate = Hit.Location - CurrentWorldPos;
+					}
+					else
+					{
+						FVector ArriveTangent = SplineComp->GetComponentTransform().GetRotation().RotateVector(EditedPoint.ArriveTangent); // convert local-space tangent vector to world-space
+						DeltaTranslate = ArriveTangent.GetSafeNormal() * ArriveTangent.Size() * 0.5;
+						DeltaTranslate = ArriveTangent.GetSafeNormal() * ArriveTangent.Size() * 0.5;
+					}
+
+					OnAddKey();
+					TransformSelectedKeys(DeltaTranslate);
+
+					return true;
+				}
+			}
+		}
+	}
+	*/
+	return false;
+}
+
+
+bool FSplineComponentVisualizer::HandleBoxSelect(const FBox& InBox, FEditorViewportClient* InViewportClient, FViewport* InViewport) 
+{
+	USplineComponent* SplineComp = GetEditedSplineComponent();
+	if (SplineComp != nullptr)
+	{
+		bool bSelectionChanged = false;
+
+		const FInterpCurveVector& SplineInfo = SplineComp->GetSplinePointsPosition();
+		int32 NumPoints = SplineInfo.Points.Num();
+
+		// Spline control point selection always uses transparent box selection.
+		for (int32 KeyIdx = 0; KeyIdx < NumPoints; KeyIdx++)
+		{
+			const FVector Pos = SplineComp->GetLocationAtSplinePoint(KeyIdx, ESplineCoordinateSpace::World);
+
+			if (InBox.IsInside(Pos))
+			{
+				ChangeSelectionState(KeyIdx, true);
+				bSelectionChanged = true;
+			}
+		}
+
+		if (bSelectionChanged)
+		{
+			SelectedSegmentIndex = INDEX_NONE;
+			SelectedTangentHandle = INDEX_NONE;
+			SelectedTangentHandleType = ESelectedTangentHandle::None;
+		}
+	}
+
+	return true;
+}
+
+bool FSplineComponentVisualizer::HandleFrustumSelect(const FConvexVolume& InFrustum, FEditorViewportClient* InViewportClient, FViewport* InViewport) 
+{
+	USplineComponent* SplineComp = GetEditedSplineComponent();
+	if (SplineComp != nullptr)
+	{
+		bool bSelectionChanged = false;
+
+		const FInterpCurveVector& SplineInfo = SplineComp->GetSplinePointsPosition();
+		int32 NumPoints = SplineInfo.Points.Num();
+
+		// Spline control point selection always uses transparent box selection.
+		for (int32 KeyIdx = 0; KeyIdx < NumPoints; KeyIdx++)
+		{
+			const FVector Pos = SplineComp->GetLocationAtSplinePoint(KeyIdx, ESplineCoordinateSpace::World);
+
+			if (InFrustum.IntersectPoint(Pos))
+			{
+				ChangeSelectionState(KeyIdx, true);
+				bSelectionChanged = true;
+			}
+		}
+
+		if (bSelectionChanged)
+		{
+			SelectedSegmentIndex = INDEX_NONE;
+			SelectedTangentHandle = INDEX_NONE;
+			SelectedTangentHandleType = ESelectedTangentHandle::None;
+		}
+	}
+
+	return true;
+}
+
+bool FSplineComponentVisualizer::HasFocusOnSelectionBoundingBox(FBox& OutBoundingBox)
+{
+	OutBoundingBox.Init();
+
+	if (SelectedKeys.Num() > 0)
+	{
+		USplineComponent* SplineComp = GetEditedSplineComponent();
+		if (SplineComp != nullptr)
+		{
+			// Spline control point selection always uses transparent box selection.
+			for (int32 KeyIdx : SelectedKeys)
+			{
+				const FVector Pos = SplineComp->GetLocationAtSplinePoint(KeyIdx, ESplineCoordinateSpace::World);
+
+				OutBoundingBox += Pos;
+			}
+
+			OutBoundingBox.ExpandBy(50.f);
+			return true;
+		}
+	}
+
+	return false;
+}
+
+bool FSplineComponentVisualizer::HandleSnapTo(const bool bInAlign, const bool bInUseLineTrace, const bool bInUseBounds, const bool bInUsePivot, AActor* InDestination)
+{
+	// Does not handle Snap/Align Pivot, Snap/Align Bottom Control Points or Snap/Align to Actor.
+	if (bInUsePivot || bInUseBounds || InDestination)
+	{
+		return false;
+	}
+
+	// Note: value of bInUseLineTrace is ignored as we always line trace from control points.
+
+	USplineComponent* SplineComp = GetEditedSplineComponent();
+
+	if (SplineComp != nullptr)
+	{
+		if (SelectedKeys.Num() > 0)
+		{
+			check(LastKeyIndexSelected != INDEX_NONE);
+			check(SelectedKeys.Contains(LastKeyIndexSelected));
+
+			SplineComp->Modify();
+
+			FInterpCurveVector& SplinePosition = SplineComp->GetSplinePointsPosition();
+			FInterpCurveQuat& SplineRotation = SplineComp->GetSplinePointsRotation();
+			int32 NumPoints = SplinePosition.Points.Num();
+
+			bool bMovedKey = false;
+
+			// Spline control point selection always uses transparent box selection.
+			for (int32 KeyIdx : SelectedKeys)
+			{
+				check(KeyIdx < NumPoints);
+
+				FVector Direction = FVector(0.f, 0.f, -1.f);
+
+				FInterpCurvePoint<FVector>& EditedPoint = SplinePosition.Points[KeyIdx];
+				FInterpCurvePoint<FQuat>& EditedRotPoint = SplineRotation.Points[KeyIdx];
+
+				FHitResult Hit(1.0f);
+				FCollisionQueryParams Params(SCENE_QUERY_STAT(MoveSplineKeyToTrace), true);
+
+				// Find key position in world space
+				const FVector CurrentWorldPos = SplineComp->GetComponentTransform().TransformPosition(EditedPoint.OutVal);
+
+				if (SplineComp->GetWorld()->LineTraceSingleByChannel(Hit, CurrentWorldPos, CurrentWorldPos + Direction * WORLD_MAX, ECC_WorldStatic, Params))
+				{
+					// Convert back to local space
+					EditedPoint.OutVal = SplineComp->GetComponentTransform().InverseTransformPosition(Hit.Location);
+
+					if (bInAlign)
+					{		
+						// Set point tangent as user controlled
+						EditedPoint.InterpMode = CIM_CurveUser;
+
+						// Get delta rotation between up vector and hit normal
+						FVector WorldUpVector = SplineComp->GetUpVectorAtSplineInputKey(KeyIdx, ESplineCoordinateSpace::World);
+						FQuat DeltaRotate = FQuat::FindBetweenNormals(WorldUpVector, Hit.Normal);
+
+						// Rotate tangent according to delta rotation
+						FVector NewTangent = SplineComp->GetComponentTransform().GetRotation().RotateVector(EditedPoint.LeaveTangent); // convert local-space tangent vector to world-space
+						NewTangent = DeltaRotate.RotateVector(NewTangent); // apply world-space delta rotation to world-space tangent
+						NewTangent = SplineComp->GetComponentTransform().GetRotation().Inverse().RotateVector(NewTangent); // convert world-space tangent vector back into local-space
+						EditedPoint.LeaveTangent = NewTangent;
+						EditedPoint.ArriveTangent = NewTangent;
+
+						// Rotate spline rotation according to delta rotation
+						FQuat NewRot = SplineComp->GetComponentTransform().GetRotation() * EditedRotPoint.OutVal; // convert local-space rotation to world-space
+						NewRot = DeltaRotate * NewRot; // apply world-space rotation
+						NewRot = SplineComp->GetComponentTransform().GetRotation().Inverse() * NewRot; // convert world-space rotation to local-space
+						EditedRotPoint.OutVal = NewRot;
+					}
+
+					bMovedKey = true;
+				}
+			}
+
+			if (bMovedKey)
+			{
+				SplineComp->UpdateSpline();
+				SplineComp->bSplineHasBeenEdited = true;
+
+				NotifyPropertyModified(SplineComp, SplineCurvesProperty);
+				
+				if (bInAlign)
+				{
+					CachedRotation = SplineComp->GetQuaternionAtSplinePoint(LastKeyIndexSelected, ESplineCoordinateSpace::World);
+				}
+
+				GEditor->RedrawLevelEditingViewports(true);
+			}
+
+			return true;
+		}
+	}
+
+	return false;
+}
+
+void FSplineComponentVisualizer::OnSnapToMarkedKey(bool bAlign)
+{
+	const FScopedTransaction Transaction(LOCTEXT("SnapToMarkedSplinePoint", "Snap To Marked Spline Point"));
+
+	USplineComponent* SplineComp = GetEditedSplineComponent();
+	check(SplineComp != nullptr);
+	check(LastKeyIndexSelected != INDEX_NONE);
+	check(SelectedKeys.Num() == 1);
+	check(SelectedKeys.Contains(LastKeyIndexSelected));
+	check(HasMarkedKey());
+	check(!IsMarkedKey(SplineComp, LastKeyIndexSelected));
+
+	SplineComp->Modify();
+	if (AActor* Owner = SplineComp->GetOwner())
+	{
+		Owner->Modify();
+	}
+
+	const USplineComponent* MarkedSplineComp = MarkedKeySplineComponent.Get();
+	check(MarkedSplineComp != nullptr);
+
+	const FInterpCurvePoint<FVector>& MarkedPosition = MarkedSplineComp->GetSplinePointsPosition().Points[MarkedKeyIndex];
+	const FInterpCurvePoint<FQuat>& MarkedRotation = MarkedSplineComp->GetSplinePointsRotation().Points[MarkedKeyIndex];
+	const FInterpCurvePoint<FVector>& MarkedScale = MarkedSplineComp->GetSplinePointsScale().Points[MarkedKeyIndex];
+
+	FInterpCurvePoint<FVector>& EditedPosition = SplineComp->GetSplinePointsPosition().Points[LastKeyIndexSelected];
+	FInterpCurvePoint<FQuat>& EditedRotation = SplineComp->GetSplinePointsRotation().Points[LastKeyIndexSelected];
+	FInterpCurvePoint<FVector>& EditedScale = SplineComp->GetSplinePointsScale().Points[LastKeyIndexSelected];
+
+	// Copy position
+	const FVector NewWorldPos = MarkedSplineComp->GetComponentTransform().TransformPosition(MarkedPosition.OutVal); // convert local-space position to world-space
+	EditedPosition.OutVal = SplineComp->GetComponentTransform().InverseTransformPosition(NewWorldPos); // convert world-space position to local-space
+
+	if (bAlign)
+	{
+		// Copy tangents
+		EditedPosition.InterpMode = CIM_CurveUser;
+		const FVector NewArriveTangent = MarkedSplineComp->GetComponentTransform().GetRotation().RotateVector(MarkedPosition.ArriveTangent); // convert local-space tangent vectors to world-space
+		const FVector NewLeaveTangent = MarkedSplineComp->GetComponentTransform().GetRotation().RotateVector(MarkedPosition.LeaveTangent); 
+		EditedPosition.ArriveTangent = SplineComp->GetComponentTransform().GetRotation().Inverse().RotateVector(NewArriveTangent); // convert world-space tangent vectors back into local-space
+		EditedPosition.LeaveTangent = SplineComp->GetComponentTransform().GetRotation().Inverse().RotateVector(NewLeaveTangent); 
+  
+		// Copy rotation
+		const FQuat NewRot = MarkedSplineComp->GetComponentTransform().GetRotation() * MarkedRotation.OutVal; // convert local-space rotation to world-space
+		EditedRotation.OutVal = SplineComp->GetComponentTransform().GetRotation().Inverse() * NewRot; // convert world-space rotation to local-space
+	}
+  
+	// Copy scale - X is not used so ignore it
+	const FVector MarkedSplineCompScale = MarkedSplineComp->GetComponentTransform().GetScale3D();
+	const FVector SplineCompScale = SplineComp->GetComponentTransform().GetScale3D();
+	float NewScaleY = MarkedSplineCompScale.Y * MarkedScale.OutVal.Y; // convert local-space scale to world-space
+	float NewScaleZ = MarkedSplineCompScale.Z * MarkedScale.OutVal.Z;
+	EditedScale.OutVal.Y = FMath::IsNearlyZero(SplineCompScale.Y) ? NewScaleY : NewScaleY / SplineCompScale.Y; // convert world-space scale to local-space
+	EditedScale.OutVal.Z = FMath::IsNearlyZero(SplineCompScale.Z) ? NewScaleZ : NewScaleZ / SplineCompScale.Z; 
+
+	// Copy metadata
+	USplineMetadata* SplineMetadata = SplineComp->GetSplinePointsMetadata();
+	if (SplineMetadata)
+	{
+		SplineMetadata->CopyPoint(MarkedSplineComp, MarkedKeyIndex, LastKeyIndexSelected);
+	}
+
+	SplineComp->UpdateSpline();
+	SplineComp->bSplineHasBeenEdited = true;
+
+	NotifyPropertyModified(SplineComp, SplineCurvesProperty);
+				
+	if (bAlign)
+	{
+		CachedRotation = SplineComp->GetQuaternionAtSplinePoint(LastKeyIndexSelected, ESplineCoordinateSpace::World);
+	}
+
+	GEditor->RedrawLevelEditingViewports(true);
+
+	OnResetMarkedKey();
+}
+
+bool FSplineComponentVisualizer::CanSnapToMarkedKey() const
+{
+	USplineComponent* SplineComp = GetEditedSplineComponent();
+	return (SplineComp != nullptr &&
+		SelectedKeys.Num() == 1 &&
+		LastKeyIndexSelected != INDEX_NONE &&
+		HasMarkedKey() &&
+		!IsMarkedKey(SplineComp, LastKeyIndexSelected));
+}
+
+void FSplineComponentVisualizer::OnSnapAll(EAxis::Type InAxis)
+{
+	/*
+	const FScopedTransaction Transaction(LOCTEXT("DeleteSplinePoint", "Delete Spline Point"));
+	USplineComponent* SplineComp = GetEditedSplineComponent();
+	check(SplineComp != nullptr);
+	check(LastKeyIndexSelected != INDEX_NONE);
+	check(SelectedKeys.Num() > 0);
+	check(SelectedKeys.Contains(LastKeyIndexSelected));
+
+	SplineComp->Modify();
+	if (AActor* Owner = SplineComp->GetOwner())
+	{
+		Owner->Modify();
+	}
+
+	// Get a sorted list of all the selected indices, highest to lowest
+	TArray<int32> SelectedKeysSorted;
+	for (int32 SelectedKeyIndex : SelectedKeys)
+	{
+		SelectedKeysSorted.Add(SelectedKeyIndex);
+	}
+	SelectedKeysSorted.Sort([](int32 A, int32 B) { return A > B; });
+
+	// Delete selected keys from list, highest index first
+	FInterpCurveVector& SplinePosition = SplineComp->GetSplinePointsPosition();
+	FInterpCurveQuat& SplineRotation = SplineComp->GetSplinePointsRotation();
+	FInterpCurveVector& SplineScale = SplineComp->GetSplinePointsScale();
+	USplineMetadata* SplineMetadata = SplineComp->GetSplinePointsMetadata();
+		
+	for (int32 SelectedKeyIndex : SelectedKeysSorted)
+	{
+		if (SplineMetadata)
+		{
+			SplineMetadata->RemovePoint(SelectedKeyIndex);
+		}
+		
+		SplinePosition.Points.RemoveAt(SelectedKeyIndex);
+		SplineRotation.Points.RemoveAt(SelectedKeyIndex);
+		SplineScale.Points.RemoveAt(SelectedKeyIndex);
+
+		for (int Index = SelectedKeyIndex; Index < SplinePosition.Points.Num(); Index++)
+		{
+			SplinePosition.Points[Index].InVal -= 1.0f;
+			SplineRotation.Points[Index].InVal -= 1.0f;
+			SplineScale.Points[Index].InVal -= 1.0f;
+		}
+	}
+
+	// Select first key
+	ChangeSelectionState(0, false);
+	SelectedSegmentIndex = INDEX_NONE;
+	SelectedTangentHandle = INDEX_NONE;
+	SelectedTangentHandleType = ESelectedTangentHandle::None;
+
+	SplineComp->UpdateSpline();
+	SplineComp->bSplineHasBeenEdited = true;
+
+	NotifyPropertyModified(SplineComp, SplineCurvesProperty);
+
+	CachedRotation = SplineComp->GetQuaternionAtSplinePoint(LastKeyIndexSelected, ESplineCoordinateSpace::World);
+
+	GEditor->RedrawLevelEditingViewports(true);
+	*/
+}
+
+bool FSplineComponentVisualizer::CanSnapAll() const
+{
+	USplineComponent* SplineComp = GetEditedSplineComponent();
+	return (SplineComp != nullptr &&
+			SelectedKeys.Num() == 1 &&
+			LastKeyIndexSelected != INDEX_NONE);
+}
+
 
 void FSplineComponentVisualizer::EndEditing()
 {
@@ -876,7 +1448,6 @@ bool FSplineComponentVisualizer::CanAddKey() const
 	return (SelectedSegmentIndex != INDEX_NONE && SelectedSegmentIndex < NumSegments);
 }
 
-
 void FSplineComponentVisualizer::OnAddKey()
 {
 	const FScopedTransaction Transaction(LOCTEXT("AddSplinePoint", "Add Spline Point"));
@@ -982,6 +1553,11 @@ void FSplineComponentVisualizer::OnDeleteKey()
 		
 	for (int32 SelectedKeyIndex : SelectedKeysSorted)
 	{
+		if (IsMarkedKey(SplineComp, SelectedKeyIndex))
+		{
+			OnResetMarkedKey();
+		}
+
 		if (SplineMetadata)
 		{
 			SplineMetadata->RemovePoint(SelectedKeyIndex);
@@ -1034,6 +1610,91 @@ bool FSplineComponentVisualizer::IsKeySelectionValid() const
 			LastKeyIndexSelected != INDEX_NONE);
 }
 
+void FSplineComponentVisualizer::OnSetMarkedKey()
+{
+	const FScopedTransaction Transaction(LOCTEXT("SetMarkedSplinePoint", "Set Marked Spline Point"));
+	USplineComponent* SplineComp = GetEditedSplineComponent();
+	check(SplineComp != nullptr);
+	check(LastKeyIndexSelected != INDEX_NONE);
+	check(SelectedKeys.Num() > 0);
+	check(SelectedKeys.Contains(LastKeyIndexSelected));
+
+	if (SelectedKeys.Num() == 1)
+	{
+		MarkedKeyIndex = LastKeyIndexSelected;
+		MarkedKeySplineComponent = SplineComp;
+	}
+}
+
+bool FSplineComponentVisualizer::CanSetMarkedKey() const
+{
+	USplineComponent* SplineComp = GetEditedSplineComponent();
+	return (SplineComp != nullptr &&
+			SelectedKeys.Num() == 1 &&
+			LastKeyIndexSelected != INDEX_NONE);
+}
+
+void FSplineComponentVisualizer::OnResetMarkedKey()
+{
+	MarkedKeyIndex = INDEX_NONE;
+	MarkedKeySplineComponent = nullptr;
+}
+
+bool FSplineComponentVisualizer::CanResetMarkedKey() const
+{
+	return (MarkedKeyIndex != INDEX_NONE && MarkedKeySplineComponent.IsValid());
+}
+
+bool FSplineComponentVisualizer::ContainsMarkedKey(const USplineComponent* InSplineComp) const
+{
+	return (InSplineComp != nullptr &&
+			HasMarkedKey() &&
+			InSplineComp == MarkedKeySplineComponent.Get());
+}
+
+bool FSplineComponentVisualizer::IsMarkedKey(const USplineComponent* InSplineComp, int32 InKeyIdx) const
+{
+	return (ContainsMarkedKey(InSplineComp) && InKeyIdx == MarkedKeyIndex);
+}
+
+bool FSplineComponentVisualizer::HasMarkedKey() const
+{
+	return (MarkedKeyIndex != INDEX_NONE &&
+			MarkedKeySplineComponent.IsValid() &&
+			MarkedKeyIndex < MarkedKeySplineComponent.Get()->GetSplinePointsPosition().Points.Num());
+}
+
+/*
+void FSplineComponentVisualizer::ValidateMarkedKey() 
+{
+	if (MarkedKeySplineComponent.IsValid())
+	{
+		if (MarkedKeyIndex == INDEX_NONE ||
+			MarkedKeyIndex >= MarkedKeySplineComponent.Get()->GetSplinePointsPosition().Points.Num())
+		{
+			MarkedKeySplineComponent = nullptr;
+			MarkedKeyIndex = INDEX_NONE;
+		}
+	}
+	else
+	{
+		if (MarkedKeyIndex != INDEX_NONE)
+		{
+			MarkedKeyIndex = INDEX_NONE;
+		}
+	}
+}
+*/
+
+void FSplineComponentVisualizer::OnLockAxis(EAxis::Type InAxis)
+{
+	AddKeyLockedAxis = InAxis;
+}
+
+bool FSplineComponentVisualizer::IsLockAxisSet(EAxis::Type Index) const
+{
+	return (Index == AddKeyLockedAxis);
+}
 
 void FSplineComponentVisualizer::OnResetToAutomaticTangent(EInterpCurveMode Mode)
 {
@@ -1286,6 +1947,29 @@ TSharedPtr<SWidget> FSplineComponentVisualizer::GenerateContextMenu() const
 		}
 		MenuBuilder.EndSection();
 
+		MenuBuilder.BeginSection("Transform");
+		{
+			MenuBuilder.AddMenuEntry(FEditorViewportCommands::Get().FocusViewportToSelection);
+
+			MenuBuilder.AddSubMenu(
+				LOCTEXT("SnapAlign", "Snap/Align"),
+				LOCTEXT("KeyTypeTooltip", "Snap align options."),
+				FNewMenuDelegate::CreateSP(this, &FSplineComponentVisualizer::GenerateSnapAlignSubMenu));
+
+			MenuBuilder.AddSubMenu(
+				LOCTEXT("Mark", "Mark"),
+				LOCTEXT("KeyTypeTooltip", "Mark selected spline point."),
+				FNewMenuDelegate::CreateSP(this, &FSplineComponentVisualizer::GenerateMarkSubMenu));
+
+			/* temporarily disabled
+			MenuBuilder.AddSubMenu(
+				LOCTEXT("LockAxis", "Lock Axis"),
+				LOCTEXT("KeyTypeTooltip", "Axis to lock when adding new spline points."),
+				FNewMenuDelegate::CreateSP(this, &FSplineComponentVisualizer::GenerateLockAxisSubMenu));
+				*/
+		}
+		MenuBuilder.EndSection();
+
 		MenuBuilder.BeginSection("Spline", LOCTEXT("Spline", "Spline"));
 		{
 			MenuBuilder.AddMenuEntry(FSplineComponentVisualizerCommands::Get().ResetToDefault);
@@ -1304,7 +1988,6 @@ TSharedPtr<SWidget> FSplineComponentVisualizer::GenerateContextMenu() const
 	return MenuWidget;
 }
 
-
 void FSplineComponentVisualizer::GenerateSplinePointTypeSubMenu(FMenuBuilder& MenuBuilder) const
 {
 	MenuBuilder.AddMenuEntry(FSplineComponentVisualizerCommands::Get().SetKeyToCurve);
@@ -1312,11 +1995,39 @@ void FSplineComponentVisualizer::GenerateSplinePointTypeSubMenu(FMenuBuilder& Me
 	MenuBuilder.AddMenuEntry(FSplineComponentVisualizerCommands::Get().SetKeyToConstant);
 }
 
-
 void FSplineComponentVisualizer::GenerateTangentTypeSubMenu(FMenuBuilder& MenuBuilder) const
 {
 	MenuBuilder.AddMenuEntry(FSplineComponentVisualizerCommands::Get().ResetToUnclampedTangent);
 	MenuBuilder.AddMenuEntry(FSplineComponentVisualizerCommands::Get().ResetToClampedTangent);
 }
 
+void FSplineComponentVisualizer::GenerateSnapAlignSubMenu(FMenuBuilder& MenuBuilder) const
+{
+	MenuBuilder.AddMenuEntry(FLevelEditorCommands::Get().SnapToFloor);
+	MenuBuilder.AddMenuEntry(FLevelEditorCommands::Get().AlignToFloor);
+	MenuBuilder.AddMenuEntry(FSplineComponentVisualizerCommands::Get().SnapToMarkedKey);
+	MenuBuilder.AddMenuEntry(FSplineComponentVisualizerCommands::Get().AlignToMarkedKey);
+	/* temporarily disable
+	MenuBuilder.AddMenuEntry(FSplineComponentVisualizerCommands::Get().SnapAllToSelectedX);
+	MenuBuilder.AddMenuEntry(FSplineComponentVisualizerCommands::Get().SnapAllToSelectedY);
+	MenuBuilder.AddMenuEntry(FSplineComponentVisualizerCommands::Get().SnapAllToSelectedZ);
+	*/
+}
+
+void FSplineComponentVisualizer::GenerateMarkSubMenu(FMenuBuilder& MenuBuilder) const
+{
+	MenuBuilder.AddMenuEntry(FSplineComponentVisualizerCommands::Get().SetMarkedKey);
+	MenuBuilder.AddMenuEntry(FSplineComponentVisualizerCommands::Get().ResetMarkedKey);
+}
+
+void FSplineComponentVisualizer::GenerateLockAxisSubMenu(FMenuBuilder& MenuBuilder) const
+{
+	MenuBuilder.AddMenuEntry(FSplineComponentVisualizerCommands::Get().SetLockedAxisNone);
+	MenuBuilder.AddMenuEntry(FSplineComponentVisualizerCommands::Get().SetLockedAxisX);
+	MenuBuilder.AddMenuEntry(FSplineComponentVisualizerCommands::Get().SetLockedAxisY);
+	MenuBuilder.AddMenuEntry(FSplineComponentVisualizerCommands::Get().SetLockedAxisZ);
+}
+
+
+
 #undef LOCTEXT_NAMESPACE
diff --git a/Engine/Source/Editor/ComponentVisualizers/Public/SplineComponentVisualizer.h b/Engine/Source/Editor/ComponentVisualizers/Public/SplineComponentVisualizer.h
--- a/Engine/Source/Editor/ComponentVisualizers/Public/SplineComponentVisualizer.h
+++ b/Engine/Source/Editor/ComponentVisualizers/Public/SplineComponentVisualizer.h
@@ -6,6 +6,7 @@
 #include "InputCoreTypes.h"
 #include "HitProxies.h"
 #include "ComponentVisualizer.h"
+#include "Components/SplineComponent.h"
 
 class AActor;
 class FEditorViewportClient;
@@ -17,6 +18,7 @@ class FViewport;
 class SWidget;
 class USplineComponent;
 struct FViewportClick;
+struct FConvexVolume;
 
 /** Base class for clickable spline editing proxies */
 struct HSplineVisProxy : public HComponentVisProxy
@@ -85,6 +87,16 @@ public:
 	virtual bool GetCustomInputCoordinateSystem(const FEditorViewportClient* ViewportClient, FMatrix& OutMatrix) const override;
 	virtual bool HandleInputDelta(FEditorViewportClient* ViewportClient, FViewport* Viewport, FVector& DeltaTranslate, FRotator& DeltaRotate, FVector& DeltaScale) override;
 	virtual bool HandleInputKey(FEditorViewportClient* ViewportClient, FViewport* Viewport, FKey Key, EInputEvent Event) override;
+	/** Handle click modified by Alt, Ctrl and/or Shift. The input HitProxy may not be on this component. */
+	virtual bool HandleModifiedClick(FEditorViewportClient* InViewportClient, HHitProxy* HitProxy, const FViewportClick& Click) override;
+	/** Handle box select input */
+	virtual bool HandleBoxSelect(const FBox& InBox, FEditorViewportClient* InViewportClient, FViewport* InViewport) override;
+	/** Handle frustum select input */
+	virtual bool HandleFrustumSelect(const FConvexVolume& InFrustum, FEditorViewportClient* InViewportClient, FViewport* InViewport) override;
+	/** Return whether focus on selection should focus on bounding box defined by active visualizer */
+	virtual bool HasFocusOnSelectionBoundingBox(FBox& OutBoundingBox) override;
+	/** Pass snap input to active visualizer */
+	virtual bool HandleSnapTo(const bool bInAlign, const bool bInUseLineTrace, const bool bInUseBounds, const bool bInUsePivot, AActor* InDestination) override;
 	virtual TSharedPtr<SWidget> GenerateContextMenu() const override;
 	virtual bool IsVisualizingArchetype() const override;
 	//~ End FComponentVisualizer Interface
@@ -96,6 +108,16 @@ public:
 
 protected:
 
+	/** Transforms selected tangent by given translation */
+	bool TransformSelectedTangent(const FVector& DeltaTranslate);
+
+	/** Transforms selected tangent by given translate, rotate and scale */
+	bool TransformSelectedKeys(const FVector& DeltaTranslate, const FRotator& DeltaRotate = FRotator::ZeroRotator, const FVector& DeltaScale = FVector::ZeroVector, bool bDuplicateKey = false);
+
+	/** Snaps edited point and rot point to given world location, up vector, and tangent */
+	void SnapTo(USplineComponent *SplineComp, int32 KeyIdx,
+		const FVector& InLocation, const FVector& InUpVector = FVector::ZeroVector, bool bInAlignUpVector = false, const FVector& InTangentVector = FVector::ZeroVector, bool bInAlignTangentVector = false);
+
 	/** Update the key selection state of the visualizer */
 	void ChangeSelectionState(int32 Index, bool bIsCtrlHeld);
 
@@ -111,6 +133,21 @@ protected:
 	void OnAddKey();
 	bool CanAddKey() const;
 
+	void OnSetMarkedKey();
+	bool CanSetMarkedKey() const;
+
+	void OnResetMarkedKey();
+	bool CanResetMarkedKey() const;
+
+	void OnSnapToMarkedKey(bool bAlign);
+	bool CanSnapToMarkedKey() const;
+
+	void OnSnapAll(EAxis::Type InAxis);
+	bool CanSnapAll() const;
+
+	void OnLockAxis(EAxis::Type InAxis);
+	bool IsLockAxisSet(EAxis::Type InAxis) const; 
+	
 	void OnResetToAutomaticTangent(EInterpCurveMode Mode);
 	bool CanResetToAutomaticTangent(EInterpCurveMode Mode) const;
 
@@ -126,12 +163,31 @@ protected:
 	void OnResetToDefault();
 	bool CanResetToDefault() const;
 
+
+	/** Returns true if specified spline component contains the marked key */
+	bool ContainsMarkedKey(const USplineComponent* InSplineComp) const;
+
+	/** Returns true if specified spline component and key index are the marked key */
+	bool IsMarkedKey(const USplineComponent* InSplineComp, int32 InKeyIdx) const;
+
+	/** Returns true if there is a valid marked key */
+	bool HasMarkedKey() const;
+
 	/** Generate the submenu containing the available point types */
 	void GenerateSplinePointTypeSubMenu(FMenuBuilder& MenuBuilder) const;
 
 	/** Generate the submenu containing the available auto tangent types */
 	void GenerateTangentTypeSubMenu(FMenuBuilder& MenuBuilder) const;
 
+	/** Generate the submenu containing the available snap/align actions */
+	void GenerateSnapAlignSubMenu(FMenuBuilder& MenuBuilder) const;
+	
+	/** Generate the submenu containing the mark key actions */
+	void GenerateMarkSubMenu(FMenuBuilder& MenuBuilder) const;
+
+	/** Generate the submenu containing the lock axis types */
+	void GenerateLockAxisSubMenu(FMenuBuilder& MenuBuilder) const;
+
 	/** Output log commands */
 	TSharedPtr<FUICommandList> SplineComponentVisualizerActions;
 
@@ -175,6 +231,16 @@ protected:
 	/** Whether we currently allow duplication when dragging */
 	bool bAllowDuplication;
 
+	/** Spline component containing marked spline point, used for snapping. */
+	TWeakObjectPtr<USplineComponent> MarkedKeySplineComponent;
+
+	/** Index of marked spline point, used for snapping. Note: this is not necessarily in the current spline component. */
+	int32 MarkedKeyIndex;
+
+	/** Axis to fix when adding new spline points. Uses the value of the currently 
+	    selected spline point's X, Y, or Z value when fix is not equal to none. */
+	EAxis::Type AddKeyLockedAxis;
+
 private:
 	UProperty* SplineCurvesProperty;
 };
diff --git a/Engine/Source/Editor/LevelEditor/Private/LevelEditorActions.cpp b/Engine/Source/Editor/LevelEditor/Private/LevelEditorActions.cpp
--- a/Engine/Source/Editor/LevelEditor/Private/LevelEditorActions.cpp
+++ b/Engine/Source/Editor/LevelEditor/Private/LevelEditorActions.cpp
@@ -2993,7 +2993,14 @@ void FLevelEditorActionCallbacks::SnapTo_Clicked( const bool InAlign, const bool
 	FScopedLevelDirtied		LevelDirtyCallback;
 
 	bool bSnappedComponents = false;
-	if( GEditor->GetSelectedComponentCount() > 0 )
+
+	// Let the component visualizers try to handle the selection.
+	if (GUnrealEd->ComponentVisManager.HandleSnapTo(InAlign, InUseLineTrace, InUseBounds, InUsePivot, InDestination))
+	{
+		bSnappedComponents = true;
+	}
+
+	if( !bSnappedComponents && GEditor->GetSelectedComponentCount() > 0 )
 	{
 		for(FSelectedEditableComponentIterator It(GEditor->GetSelectedEditableComponentIterator()); It; ++It)
 		{
diff --git a/Engine/Source/Editor/UnrealEd/Private/ComponentVisualizerManager.cpp b/Engine/Source/Editor/UnrealEd/Private/ComponentVisualizerManager.cpp
--- a/Engine/Source/Editor/UnrealEd/Private/ComponentVisualizerManager.cpp
+++ b/Engine/Source/Editor/UnrealEd/Private/ComponentVisualizerManager.cpp
@@ -17,6 +17,17 @@ FComponentVisualizerManager::FComponentVisualizerManager()
 /** Handle a click on the specified editor viewport client */
 bool FComponentVisualizerManager::HandleClick(FEditorViewportClient* InViewportClient, HHitProxy* HitProxy, const FViewportClick& Click)
 {
+	TSharedPtr<FComponentVisualizer> EditedVisualizer = EditedVisualizerPtr.Pin();
+
+	// Give current visualizer a chance to handle click if it has a modifier.
+	if (EditedVisualizer.IsValid() && (Click.IsControlDown() || Click.IsAltDown() || Click.IsShiftDown()))
+	{
+		if (EditedVisualizer->HandleModifiedClick(InViewportClient, HitProxy, Click))
+		{ 
+			return true;
+		}
+	}
+		
 	bool bHandled = HandleProxyForComponentVis(InViewportClient, HitProxy, Click);
 	if (bHandled && Click.GetKey() == EKeys::RightMouseButton)
 	{
@@ -120,6 +131,65 @@ bool FComponentVisualizerManager::HandleInputDelta(FEditorViewportClient* InView
 	return false;
 }
 
+bool FComponentVisualizerManager::HandleFrustumSelect(const FConvexVolume& InFrustum, FEditorViewportClient* InViewportClient, FViewport* InViewport) const
+{
+	TSharedPtr<FComponentVisualizer> EditedVisualizer = EditedVisualizerPtr.Pin();
+
+	if (EditedVisualizer.IsValid() && EditedVisualizerViewportClient == InViewportClient)
+	{
+		if (EditedVisualizer->HandleFrustumSelect(InFrustum, InViewportClient, InViewport))
+		{
+			return true;
+		}
+	}
+
+	return false;
+}
+
+bool FComponentVisualizerManager::HandleBoxSelect(const FBox& InBox, FEditorViewportClient* InViewportClient, FViewport* InViewport) const
+{
+	TSharedPtr<FComponentVisualizer> EditedVisualizer = EditedVisualizerPtr.Pin();
+
+	if (EditedVisualizer.IsValid() && EditedVisualizerViewportClient == InViewportClient)
+	{
+		if (EditedVisualizer->HandleBoxSelect(InBox, InViewportClient, InViewport))
+		{
+			return true;
+		}
+	}
+
+	return false;
+}
+
+bool FComponentVisualizerManager::HasFocusOnSelectionBoundingBox(FBox& OutBoundingBox) const
+{
+	TSharedPtr<FComponentVisualizer> EditedVisualizer = EditedVisualizerPtr.Pin();
+
+	if (EditedVisualizer.IsValid())
+	{
+		if (EditedVisualizer->HasFocusOnSelectionBoundingBox(OutBoundingBox))
+		{
+			return true;
+		}
+	}
+
+	return false;
+}
+
+bool FComponentVisualizerManager::HandleSnapTo(const bool bInAlign, const bool bInUseLineTrace, const bool bInUseBounds, const bool bInUsePivot, AActor* InDestination)
+{
+	TSharedPtr<FComponentVisualizer> EditedVisualizer = EditedVisualizerPtr.Pin();
+
+	if (EditedVisualizer.IsValid())
+	{
+		if (EditedVisualizer->HandleSnapTo(bInAlign, bInUseLineTrace, bInUseBounds, bInUsePivot, InDestination))
+		{
+			return true;
+		}
+	}
+
+	return false;
+}
 
 bool FComponentVisualizerManager::GetWidgetLocation(const FEditorViewportClient* ViewportClient, FVector& OutLocation) const
 {
diff --git a/Engine/Source/Editor/UnrealEd/Private/DragTool_BoxSelect.cpp b/Engine/Source/Editor/UnrealEd/Private/DragTool_BoxSelect.cpp
--- a/Engine/Source/Editor/UnrealEd/Private/DragTool_BoxSelect.cpp
+++ b/Engine/Source/Editor/UnrealEd/Private/DragTool_BoxSelect.cpp
@@ -150,8 +150,11 @@ void FDragTool_ActorBoxSelect::EndDrag()
 	// Let the editor mode try to handle the box selection.
 	const bool bEditorModeHandledBoxSelection = ModeTools->BoxSelect(SelBBox, bLeftMouseButtonDown);
 
+	// Let the component visualizers try to handle the selection.
+	const bool bComponentVisHandledSelection = !bEditorModeHandledBoxSelection && GUnrealEd->ComponentVisManager.HandleBoxSelect(SelBBox, LevelViewportClient, LevelViewportClient->Viewport);
+
 	// If the edit mode didn't handle the selection, try normal actor box selection.
-	if ( !bEditorModeHandledBoxSelection )
+	if ( !bEditorModeHandledBoxSelection && !bComponentVisHandledSelection )
 	{
 		const bool bStrictDragSelection = GetDefault<ULevelEditorViewportSettings>()->bStrictBoxSelection;
 
diff --git a/Engine/Source/Editor/UnrealEd/Private/DragTool_FrustumSelect.cpp b/Engine/Source/Editor/UnrealEd/Private/DragTool_FrustumSelect.cpp
--- a/Engine/Source/Editor/UnrealEd/Private/DragTool_FrustumSelect.cpp
+++ b/Engine/Source/Editor/UnrealEd/Private/DragTool_FrustumSelect.cpp
@@ -77,7 +77,10 @@ void FDragTool_ActorFrustumSelect::EndDrag()
 	// Let the editor mode try to handle the selection.
 	const bool bEditorModeHandledSelection = ModeTools->FrustumSelect(Frustum, LevelViewportClient, bLeftMouseButtonDown);
 
-	if( !bEditorModeHandledSelection )
+	// Let the component visualizers try to handle the selection.
+	const bool bComponentVisHandledSelection = !bEditorModeHandledSelection && GUnrealEd->ComponentVisManager.HandleFrustumSelect(Frustum, LevelViewportClient, LevelViewportClient->Viewport);
+
+	if( !bEditorModeHandledSelection && !bComponentVisHandledSelection)
 	{
 		if( !bShiftDown )
 		{
diff --git a/Engine/Source/Editor/UnrealEd/Private/EditorServer.cpp b/Engine/Source/Editor/UnrealEd/Private/EditorServer.cpp
--- a/Engine/Source/Editor/UnrealEd/Private/EditorServer.cpp
+++ b/Engine/Source/Editor/UnrealEd/Private/EditorServer.cpp
@@ -5037,33 +5037,43 @@ bool UEditorEngine::Exec_Camera( const TCHAR* Str, FOutputDevice& Ar )
 		}
 		else 
 		{
-			TArray<AActor*> Actors;
-			for ( FSelectionIterator It( GetSelectedActorIterator() ) ; It ; ++It )
+			FBox ComponentVisBoundingBox;
+			bool bComponentVisHasFocusOnSelectionBoundingBox = GUnrealEd && GUnrealEd->ComponentVisManager.HasFocusOnSelectionBoundingBox(ComponentVisBoundingBox);
+
+			if (bComponentVisHasFocusOnSelectionBoundingBox)
 			{
-				AActor* Actor = static_cast<AActor*>( *It );
-				checkSlow( Actor->IsA(AActor::StaticClass()) );
-				Actors.Add( Actor );
+				MoveViewportCamerasToBox(ComponentVisBoundingBox, bActiveViewportOnly);
 			}
-
-			TArray<UPrimitiveComponent*> SelectedComponents;
-			for( FSelectionIterator It( GetSelectedComponentIterator() ); It; ++It )
+			else
 			{
-				UPrimitiveComponent* PrimitiveComp = Cast<UPrimitiveComponent>( *It );
-				if( PrimitiveComp )
+				TArray<AActor*> Actors;
+				for (FSelectionIterator It(GetSelectedActorIterator()); It; ++It)
 				{
-					SelectedComponents.Add( PrimitiveComp );
+					AActor* Actor = static_cast<AActor*>(*It);
+					checkSlow(Actor->IsA(AActor::StaticClass()));
+					Actors.Add(Actor);
 				}
-			}
 
-			if( Actors.Num() || SelectedComponents.Num() )
-			{
-				MoveViewportCamerasToActor( Actors, SelectedComponents, bActiveViewportOnly );
-				return true;
-			}
-			else
-			{
-				Ar.Log( TEXT("Can't find target actor or component.") );
-				return false;
+				TArray<UPrimitiveComponent*> SelectedComponents;
+				for (FSelectionIterator It(GetSelectedComponentIterator()); It; ++It)
+				{
+					UPrimitiveComponent* PrimitiveComp = Cast<UPrimitiveComponent>(*It);
+					if (PrimitiveComp)
+					{
+						SelectedComponents.Add(PrimitiveComp);
+					}
+				}
+
+				if (Actors.Num() || SelectedComponents.Num())
+				{
+					MoveViewportCamerasToActor(Actors, SelectedComponents, bActiveViewportOnly);
+					return true;
+				}
+				else
+				{
+					Ar.Log(TEXT("Can't find target actor or component."));
+					return false;
+				}
 			}
 		}
 	}
diff --git a/Engine/Source/Editor/UnrealEd/Public/ComponentVisualizer.h b/Engine/Source/Editor/UnrealEd/Public/ComponentVisualizer.h
--- a/Engine/Source/Editor/UnrealEd/Public/ComponentVisualizer.h
+++ b/Engine/Source/Editor/UnrealEd/Public/ComponentVisualizer.h
@@ -59,6 +59,16 @@ public:
 	virtual bool HandleInputDelta(FEditorViewportClient* ViewportClient, FViewport* Viewport, FVector& DeltaTranslate, FRotator& DeltalRotate, FVector& DeltaScale) { return false; }
 	/** */
 	virtual bool HandleInputKey(FEditorViewportClient* ViewportClient,FViewport* Viewport,FKey Key,EInputEvent Event) { return false; }
+	/** Handle click modified by Alt, Ctrl and/or Shift. The input HitProxy may not be on this component. */
+	virtual bool HandleModifiedClick(FEditorViewportClient* InViewportClient, HHitProxy* HitProxy, const FViewportClick& Click) { return false; }
+	/** Handle box select input */
+	virtual bool HandleBoxSelect(const FBox& InBox, FEditorViewportClient* InViewportClient,FViewport* InViewport) { return false; }
+	/** Handle frustum select input */
+	virtual bool HandleFrustumSelect(const FConvexVolume& InFrustum, FEditorViewportClient* InViewportClient, FViewport* InViewport) { return false; }
+	/** Return whether focus on selection should focus on bounding box defined by active visualizer */
+	virtual bool HasFocusOnSelectionBoundingBox(FBox& OutBoundingBox) { return false; }
+	/** Pass snap input to active visualizer */
+	virtual bool HandleSnapTo(const bool bInAlign, const bool bInUseLineTrace, const bool bInUseBounds, const bool bInUsePivot, AActor* InDestination) { return false;  }
 	/** */
 	virtual TSharedPtr<SWidget> GenerateContextMenu() const { return TSharedPtr<SWidget>(); }
 	/** */
diff --git a/Engine/Source/Editor/UnrealEd/Public/ComponentVisualizerManager.h b/Engine/Source/Editor/UnrealEd/Public/ComponentVisualizerManager.h
--- a/Engine/Source/Editor/UnrealEd/Public/ComponentVisualizerManager.h
+++ b/Engine/Source/Editor/UnrealEd/Public/ComponentVisualizerManager.h
@@ -34,6 +34,18 @@ public:
 	/** Pass delta input to active visualizer */
 	bool HandleInputDelta(FEditorViewportClient* InViewportClient, FViewport* InViewport, FVector& InDrag, FRotator& InRot, FVector& InScale) const;
 
+	/** Pass box select input to active visualizer */
+	bool HandleBoxSelect(const FBox& InBox, FEditorViewportClient* InViewportClient, FViewport* InViewport) const;
+
+	/** Pass frustum select input to active visualizer */
+	bool HandleFrustumSelect(const FConvexVolume &InFrustum, FEditorViewportClient* InViewportClient, FViewport* InViewport) const;
+
+	/** Return whether focus on selection should focus on bounding box defined by active visualizer */
+	bool HasFocusOnSelectionBoundingBox(FBox& OutBoundingBox) const;
+
+	/** Pass snap input to active visualizer */
+	bool HandleSnapTo(const bool bInAlign, const bool bInUseLineTrace, const bool bInUseBounds, const bool bInUsePivot, AActor* InDestination);
+
 	/** Get widget location from active visualizer */
 	bool GetWidgetLocation(const FEditorViewportClient* InViewportClient, FVector& OutLocation) const;
 
diff --git a/Engine/Source/Runtime/Engine/Classes/Components/SplineComponent.h b/Engine/Source/Runtime/Engine/Classes/Components/SplineComponent.h
--- a/Engine/Source/Runtime/Engine/Classes/Components/SplineComponent.h
+++ b/Engine/Source/Runtime/Engine/Classes/Components/SplineComponent.h
@@ -47,6 +47,7 @@ public:
 	virtual void AddPoint(float InputKey) PURE_VIRTUAL(USplineMetadata::AddPoint, );
 	virtual void RemovePoint(int32 Index) PURE_VIRTUAL(USplineMetadata::RemovePoint, );
 	virtual void DuplicatePoint(int32 Index) PURE_VIRTUAL(USplineMetadata::DuplicatePoint, );
+	virtual void CopyPoint(const USplineComponent* InFromSplineComp, int32 InFromIndex, int32 InToIndex) PURE_VIRTUAL(USplineMetadata::CopyPoint, );
 	virtual void Reset(int32 NumPoints) PURE_VIRTUAL(USplineMetadata::Reset, );
 	virtual void Fixup(int32 NumPoints) PURE_VIRTUAL(USplineMetadata::Fixup, );
 };
diff --git a/Engine/Source/Runtime/Engine/Private/Components/SplineComponent.cpp b/Engine/Source/Runtime/Engine/Private/Components/SplineComponent.cpp
--- a/Engine/Source/Runtime/Engine/Private/Components/SplineComponent.cpp
+++ b/Engine/Source/Runtime/Engine/Private/Components/SplineComponent.cpp
@@ -42,7 +42,7 @@ USplineComponent::USplineComponent(const FObjectInitializer& ObjectInitializer,
 	, DefaultUpVector(FVector::UpVector)
 #if WITH_EDITORONLY_DATA
 	, EditorUnselectedSplineSegmentColor(FLinearColor(1.0f, 1.0f, 1.0f))
-	, EditorSelectedSplineSegmentColor(FLinearColor(1.0f, 0.0f, 0.0f))
+	, EditorSelectedSplineSegmentColor(FLinearColor(0.828f, 0.364f, 0.003f))
 	, bAllowDiscontinuousSpline(false)
 	, bShouldVisualizeScale(false)
 	, ScaleVisualizationWidth(30.0f)
@@ -61,7 +61,14 @@ USplineComponent::USplineComponent(const FObjectInitializer& ObjectInitializer,
 	SplineCurves.Position.Points.Emplace(1.0f, FVector(100, 0, 0), FVector::ZeroVector, FVector::ZeroVector, CIM_CurveAuto);
 	SplineCurves.Rotation.Points.Emplace(1.0f, FQuat::Identity, FQuat::Identity, FQuat::Identity, CIM_CurveAuto);
 	SplineCurves.Scale.Points.Emplace(1.0f, FVector(1.0f), FVector::ZeroVector, FVector::ZeroVector, CIM_CurveAuto);
-		
+
+#if WITH_EDITORONLY_DATA
+	if (GEngine)
+	{
+		EditorSelectedSplineSegmentColor = GEngine->GetSelectionOutlineColor();
+	}
+#endif
+
 	if (SplineCurves.Metadata)
 	{
 		const int32 NumPoints = GetNumberOfSplinePoints();
diff --git a/Engine/Source/Runtime/Engine/Private/Components/SplineMeshComponent.cpp b/Engine/Source/Runtime/Engine/Private/Components/SplineMeshComponent.cpp
--- a/Engine/Source/Runtime/Engine/Private/Components/SplineMeshComponent.cpp
+++ b/Engine/Source/Runtime/Engine/Private/Components/SplineMeshComponent.cpp
@@ -522,6 +522,9 @@ void USplineMeshComponent::UpdateRenderStateAndCollision_Internal(bool bConcurre
 	}
 
 	CachedMeshBodySetupGuid.Invalidate();
+	// spline visualizer -> tell the points its interactive -> spline comp -> to water body -> generator -> update mesh on spline mesh component
+	// release the mouse must call another post edit change so that recreate physics state at that point.
+	// make this optional - optional flag set based on whether it's interactive or not.
 	RecreatePhysicsState();
 
 	bMeshDirty = false;
