Spline component: added Snap All To Selected X/Y/Z options

diff --git a/Engine/Source/Editor/ComponentVisualizers/Private/SplineComponentVisualizer.cpp b/Engine/Source/Editor/ComponentVisualizers/Private/SplineComponentVisualizer.cpp
--- a/Engine/Source/Editor/ComponentVisualizers/Private/SplineComponentVisualizer.cpp
+++ b/Engine/Source/Editor/ComponentVisualizers/Private/SplineComponentVisualizer.cpp
@@ -1263,13 +1263,13 @@ bool FSplineComponentVisualizer::CanSnapToMarkedKey() const
 
 void FSplineComponentVisualizer::OnSnapAll(EAxis::Type InAxis)
 {
-	/*
-	const FScopedTransaction Transaction(LOCTEXT("DeleteSplinePoint", "Delete Spline Point"));
+	const FScopedTransaction Transaction(LOCTEXT("SnapAllToSelectedAxis", "Snap All To Selected Axis"));
 	USplineComponent* SplineComp = GetEditedSplineComponent();
 	check(SplineComp != nullptr);
 	check(LastKeyIndexSelected != INDEX_NONE);
-	check(SelectedKeys.Num() > 0);
+	check(SelectedKeys.Num() == 1);
 	check(SelectedKeys.Contains(LastKeyIndexSelected));
+	check(InAxis == EAxis::X || InAxis == EAxis::Y || InAxis == EAxis::Z);
 
 	SplineComp->Modify();
 	if (AActor* Owner = SplineComp->GetOwner())
@@ -1277,45 +1277,74 @@ void FSplineComponentVisualizer::OnSnapAll(EAxis::Type InAxis)
 		Owner->Modify();
 	}
 
-	// Get a sorted list of all the selected indices, highest to lowest
-	TArray<int32> SelectedKeysSorted;
-	for (int32 SelectedKeyIndex : SelectedKeys)
+	FInterpCurveVector& SplinePositions = SplineComp->GetSplinePointsPosition();
+	FInterpCurveQuat& SplineRotations = SplineComp->GetSplinePointsRotation();
+
+	const FVector WorldPos = SplineComp->GetComponentTransform().TransformPosition(SplinePositions.Points[LastKeyIndexSelected].OutVal); 
+
+	FVector NewUpVector;
+	float WorldSnapAxisValue = 0.0f;
+	if (InAxis == EAxis::X)
 	{
-		SelectedKeysSorted.Add(SelectedKeyIndex);
+		WorldSnapAxisValue = WorldPos.X;
+		NewUpVector = FVector::ForwardVector;
+	}
+	else if (InAxis == EAxis::Y)
+	{
+		WorldSnapAxisValue = WorldPos.Y;
+		NewUpVector = FVector::RightVector;
+	}
+	else
+	{
+		WorldSnapAxisValue = WorldPos.Z;
+		NewUpVector = FVector::UpVector;
 	}
-	SelectedKeysSorted.Sort([](int32 A, int32 B) { return A > B; });
-
-	// Delete selected keys from list, highest index first
-	FInterpCurveVector& SplinePosition = SplineComp->GetSplinePointsPosition();
-	FInterpCurveQuat& SplineRotation = SplineComp->GetSplinePointsRotation();
-	FInterpCurveVector& SplineScale = SplineComp->GetSplinePointsScale();
-	USplineMetadata* SplineMetadata = SplineComp->GetSplinePointsMetadata();
 		
-	for (int32 SelectedKeyIndex : SelectedKeysSorted)
+	int32 NumPoints = SplinePositions.Points.Num();
+
+	for (int32 KeyIdx = 0; KeyIdx < NumPoints; KeyIdx++)
 	{
-		if (SplineMetadata)
+		FInterpCurvePoint<FVector>& EditedPosition = SplinePositions.Points[KeyIdx];
+		FInterpCurvePoint<FQuat>& EditedRotation = SplineRotations.Points[KeyIdx];
+
+		// Copy position
+		FVector NewWorldPos = SplineComp->GetComponentTransform().TransformPosition(EditedPosition.OutVal); // convert local-space position to world-space
+		if (InAxis == EAxis::X)
 		{
-			SplineMetadata->RemovePoint(SelectedKeyIndex);
+			NewWorldPos.X = WorldSnapAxisValue;
 		}
-		
-		SplinePosition.Points.RemoveAt(SelectedKeyIndex);
-		SplineRotation.Points.RemoveAt(SelectedKeyIndex);
-		SplineScale.Points.RemoveAt(SelectedKeyIndex);
-
-		for (int Index = SelectedKeyIndex; Index < SplinePosition.Points.Num(); Index++)
+		else if (InAxis == EAxis::Y)
 		{
-			SplinePosition.Points[Index].InVal -= 1.0f;
-			SplineRotation.Points[Index].InVal -= 1.0f;
-			SplineScale.Points[Index].InVal -= 1.0f;
+			NewWorldPos.Y = WorldSnapAxisValue;
+		}
+		else
+		{
+			NewWorldPos.Z = WorldSnapAxisValue;
 		}
-	}
 
-	// Select first key
-	ChangeSelectionState(0, false);
-	SelectedSegmentIndex = INDEX_NONE;
-	SelectedTangentHandle = INDEX_NONE;
-	SelectedTangentHandleType = ESelectedTangentHandle::None;
+		EditedPosition.OutVal = SplineComp->GetComponentTransform().InverseTransformPosition(NewWorldPos); // convert world-space position to local-space
 
+		// Set point tangent as user controlled
+		EditedPosition.InterpMode = CIM_CurveUser;
+
+		// Get delta rotation between current up vector and new up vector
+		FVector WorldUpVector = SplineComp->GetUpVectorAtSplineInputKey(KeyIdx, ESplineCoordinateSpace::World);
+		FQuat DeltaRotate = FQuat::FindBetweenNormals(WorldUpVector, NewUpVector);
+
+		// Rotate tangent according to delta rotation
+		FVector NewTangent = SplineComp->GetComponentTransform().GetRotation().RotateVector(EditedPosition.LeaveTangent); // convert local-space tangent vector to world-space
+		NewTangent = DeltaRotate.RotateVector(NewTangent); // apply world-space delta rotation to world-space tangent
+		NewTangent = SplineComp->GetComponentTransform().GetRotation().Inverse().RotateVector(NewTangent); // convert world-space tangent vector back into local-space
+		EditedPosition.LeaveTangent = NewTangent;
+		EditedPosition.ArriveTangent = NewTangent;
+
+		// Rotate spline rotation according to delta rotation
+		FQuat NewRot = SplineComp->GetComponentTransform().GetRotation() * EditedRotation.OutVal; // convert local-space rotation to world-space
+		NewRot = DeltaRotate * NewRot; // apply world-space rotation
+		NewRot = SplineComp->GetComponentTransform().GetRotation().Inverse() * NewRot; // convert world-space rotation to local-space
+		EditedRotation.OutVal = NewRot;
+	}
+  
 	SplineComp->UpdateSpline();
 	SplineComp->bSplineHasBeenEdited = true;
 
@@ -1324,7 +1353,6 @@ void FSplineComponentVisualizer::OnSnapAll(EAxis::Type InAxis)
 	CachedRotation = SplineComp->GetQuaternionAtSplinePoint(LastKeyIndexSelected, ESplineCoordinateSpace::World);
 
 	GEditor->RedrawLevelEditingViewports(true);
-	*/
 }
 
 bool FSplineComponentVisualizer::CanSnapAll() const
@@ -2005,13 +2033,11 @@ void FSplineComponentVisualizer::GenerateSnapAlignSubMenu(FMenuBuilder& MenuBuil
 {
 	MenuBuilder.AddMenuEntry(FLevelEditorCommands::Get().SnapToFloor);
 	MenuBuilder.AddMenuEntry(FLevelEditorCommands::Get().AlignToFloor);
-	MenuBuilder.AddMenuEntry(FSplineComponentVisualizerCommands::Get().SnapToMarkedKey);
-	MenuBuilder.AddMenuEntry(FSplineComponentVisualizerCommands::Get().AlignToMarkedKey);
-	/* temporarily disable
 	MenuBuilder.AddMenuEntry(FSplineComponentVisualizerCommands::Get().SnapAllToSelectedX);
 	MenuBuilder.AddMenuEntry(FSplineComponentVisualizerCommands::Get().SnapAllToSelectedY);
 	MenuBuilder.AddMenuEntry(FSplineComponentVisualizerCommands::Get().SnapAllToSelectedZ);
-	*/
+	MenuBuilder.AddMenuEntry(FSplineComponentVisualizerCommands::Get().SnapToMarkedKey);
+	MenuBuilder.AddMenuEntry(FSplineComponentVisualizerCommands::Get().AlignToMarkedKey);
 }
 
 void FSplineComponentVisualizer::GenerateMarkSubMenu(FMenuBuilder& MenuBuilder) const
