Added Snap/Align To Nearest Spline Point.
Removed Snap/Align To Marked Point.

diff --git a/Engine/Source/Editor/ComponentVisualizers/Private/SplineComponentVisualizer.cpp b/Engine/Source/Editor/ComponentVisualizers/Private/SplineComponentVisualizer.cpp
--- a/Engine/Source/Editor/ComponentVisualizers/Private/SplineComponentVisualizer.cpp
+++ b/Engine/Source/Editor/ComponentVisualizers/Private/SplineComponentVisualizer.cpp
@@ -1,6 +1,7 @@
 // Copyright 1998-2019 Epic Games, Inc. All Rights Reserved.
 
 #include "SplineComponentVisualizer.h"
+#include "CoreMinimal.h"
 #include "Framework/Application/SlateApplication.h"
 #include "Framework/Commands/InputChord.h"
 #include "Framework/Commands/Commands.h"
@@ -23,6 +24,7 @@ IMPLEMENT_HIT_PROXY(HSplineSegmentProxy, HSplineVisProxy);
 IMPLEMENT_HIT_PROXY(HSplineTangentHandleProxy, HSplineVisProxy);
 
 #define LOCTEXT_NAMESPACE "SplineComponentVisualizer"
+DEFINE_LOG_CATEGORY_STATIC(LogSplineComponentVisualizer, Log, All)
 
 #define VISUALIZE_SPLINE_UPVECTORS 0
 
@@ -50,13 +52,11 @@ public:
 		UI_COMMAND(SetKeyToCurve, "Curve", "Set spline point to Curve type", EUserInterfaceActionType::RadioButton, FInputChord());
 		UI_COMMAND(SetKeyToLinear, "Linear", "Set spline point to Linear type", EUserInterfaceActionType::RadioButton, FInputChord());
 		UI_COMMAND(SetKeyToConstant, "Constant", "Set spline point to Constant type", EUserInterfaceActionType::RadioButton, FInputChord());
-		UI_COMMAND(SnapToMarkedKey, "Snap to Marked Spline Point", "Snap to previously marked spline point.", EUserInterfaceActionType::Button, FInputChord());
-		UI_COMMAND(AlignToMarkedKey, "Align to Marked Spline Point", "Align to previously marked spline point.", EUserInterfaceActionType::Button, FInputChord());
+		UI_COMMAND(SnapToNearestSplinePoint, "Snap to Nearest Spline Point", "Snap to nearest spline point.", EUserInterfaceActionType::Button, FInputChord());
+		UI_COMMAND(AlignToNearestSplinePoint, "Align to Nearest Spline Point", "Align to nearest spline point.", EUserInterfaceActionType::Button, FInputChord());
 		UI_COMMAND(SnapAllToSelectedX, "Snap All To Selected X", "Snap all spline points to selected spline point X.", EUserInterfaceActionType::Button, FInputChord());
 		UI_COMMAND(SnapAllToSelectedY, "Snap All To Selected Y", "Snap all spline points to selected spline point Y.", EUserInterfaceActionType::Button, FInputChord());
 		UI_COMMAND(SnapAllToSelectedZ, "Snap All To Selected Z", "Snap all spline points to selected spline point Z.", EUserInterfaceActionType::Button, FInputChord());
-		UI_COMMAND(SetMarkedKey, "Set as Marked Spline Point", "Mark this spline point for snapping.", EUserInterfaceActionType::Button, FInputChord());
-		UI_COMMAND(ResetMarkedKey, "Reset Marked Spline Point", "Clear marked spline point.", EUserInterfaceActionType::Button, FInputChord());
 		UI_COMMAND(SetLockedAxisNone, "None", "New spline point axis is not fixed.", EUserInterfaceActionType::RadioButton, FInputChord());
 		UI_COMMAND(SetLockedAxisX, "X", "Fix X axis when adding new spline points.", EUserInterfaceActionType::RadioButton, FInputChord());
 		UI_COMMAND(SetLockedAxisY, "Y", "Fix Y axis when adding new spline points.", EUserInterfaceActionType::RadioButton, FInputChord());
@@ -91,11 +91,11 @@ public:
 	/** Set spline key to Constant type */
 	TSharedPtr<FUICommandInfo> SetKeyToConstant;
 
-	/** Snap to marked key */
-	TSharedPtr<FUICommandInfo> SnapToMarkedKey;
+	/** Snap to nearest spline point on another spline component */
+	TSharedPtr<FUICommandInfo> SnapToNearestSplinePoint;
 
-	/** Align to marked key */
-	TSharedPtr<FUICommandInfo> AlignToMarkedKey;
+	/** Align to nearest spline point on another spline component */
+	TSharedPtr<FUICommandInfo> AlignToNearestSplinePoint;
 
 	/** Snap all spline points to selected point X */
 	TSharedPtr<FUICommandInfo> SnapAllToSelectedX;
@@ -106,12 +106,6 @@ public:
 	/** Snap all spline points to selected point Z */
 	TSharedPtr<FUICommandInfo> SnapAllToSelectedZ;
 
-	/** Mark key */
-	TSharedPtr<FUICommandInfo> SetMarkedKey;
-
-	/** Clear marked key */
-	TSharedPtr<FUICommandInfo> ResetMarkedKey;
-
 	/** No axis is locked when adding new spline points */
 	TSharedPtr<FUICommandInfo> SetLockedAxisNone;
 
@@ -141,7 +135,6 @@ FSplineComponentVisualizer::FSplineComponentVisualizer()
 	, SelectedTangentHandle(INDEX_NONE)
 	, SelectedTangentHandleType(ESelectedTangentHandle::None)
 	, bAllowDuplication(true)
-	, MarkedKeyIndex(INDEX_NONE)
 	, AddKeyLockedAxis(EAxis::None)
 {
 	FSplineComponentVisualizerCommands::Register();
@@ -199,14 +192,14 @@ void FSplineComponentVisualizer::OnRegister()
 		FIsActionChecked::CreateSP(this, &FSplineComponentVisualizer::IsKeyTypeSet, CIM_Constant));
 
 	SplineComponentVisualizerActions->MapAction(
-		Commands.SnapToMarkedKey,
-		FExecuteAction::CreateSP(this, &FSplineComponentVisualizer::OnSnapToMarkedKey, false),
-		FCanExecuteAction::CreateSP(this, &FSplineComponentVisualizer::CanSnapToMarkedKey));
+		Commands.SnapToNearestSplinePoint,
+		FExecuteAction::CreateSP(this, &FSplineComponentVisualizer::OnSnapToNearestSplinePoint, false),
+		FCanExecuteAction::CreateSP(this, &FSplineComponentVisualizer::CanSnapToNearestSplinePoint));
 
 	SplineComponentVisualizerActions->MapAction(
-		Commands.AlignToMarkedKey,
-		FExecuteAction::CreateSP(this, &FSplineComponentVisualizer::OnSnapToMarkedKey, true),
-		FCanExecuteAction::CreateSP(this, &FSplineComponentVisualizer::CanSnapToMarkedKey));
+		Commands.AlignToNearestSplinePoint,
+		FExecuteAction::CreateSP(this, &FSplineComponentVisualizer::OnSnapToNearestSplinePoint, true),
+		FCanExecuteAction::CreateSP(this, &FSplineComponentVisualizer::CanSnapToNearestSplinePoint));
 
 	SplineComponentVisualizerActions->MapAction(
 		Commands.SnapAllToSelectedX,
@@ -223,16 +216,6 @@ void FSplineComponentVisualizer::OnRegister()
 		FExecuteAction::CreateSP(this, &FSplineComponentVisualizer::OnSnapAll, EAxis::Z),
 		FCanExecuteAction::CreateSP(this, &FSplineComponentVisualizer::CanSnapAll));
 
-	SplineComponentVisualizerActions->MapAction(
-		Commands.SetMarkedKey,
-		FExecuteAction::CreateSP(this, &FSplineComponentVisualizer::OnSetMarkedKey),
-		FCanExecuteAction::CreateSP(this, &FSplineComponentVisualizer::CanSetMarkedKey));
-
-	SplineComponentVisualizerActions->MapAction(
-		Commands.ResetMarkedKey,
-		FExecuteAction::CreateSP(this, &FSplineComponentVisualizer::OnResetMarkedKey),
-		FCanExecuteAction::CreateSP(this, &FSplineComponentVisualizer::CanResetMarkedKey));
-
 	SplineComponentVisualizerActions->MapAction(
 		Commands.SetLockedAxisNone,
 		FExecuteAction::CreateSP(this, &FSplineComponentVisualizer::OnLockAxis, EAxis::None),
@@ -329,7 +312,6 @@ void FSplineComponentVisualizer::DrawVisualization(const UActorComponent* Compon
 		const FColor ReadOnlyColor = FColor(255, 0, 255, 255);
 		const FColor NormalColor = bIsSplineEditable ? FColor(SplineComp->EditorUnselectedSplineSegmentColor.ToFColor(true)) : ReadOnlyColor;
 		const FColor SelectedColor = bIsSplineEditable ? FColor(SplineComp->EditorSelectedSplineSegmentColor.ToFColor(true)) : ReadOnlyColor;
-		const FColor MarkedKeyColor = FColor(255, 0, 0, 255);
 		const float GrabHandleSize = 10.0f;
 		const float TangentHandleSize = 8.0f;
 
@@ -414,7 +396,7 @@ void FSplineComponentVisualizer::DrawVisualization(const UActorComponent* Compon
 				{
 					PDI->SetHitProxy(new HSplineKeyProxy(Component, KeyIdx));
 				}
-				PDI->DrawPoint(NewKeyPos, IsMarkedKey(SplineComp, KeyIdx) ? MarkedKeyColor : KeyColor, GrabHandleSize, SDPG_Foreground);
+				PDI->DrawPoint(NewKeyPos, KeyColor, GrabHandleSize, SDPG_Foreground);
 				PDI->SetHitProxy(NULL);
 			}
 
@@ -747,8 +729,6 @@ bool FSplineComponentVisualizer::TransformSelectedTangent(const FVector& DeltaTr
 	if (SplineComp != nullptr)
 	{
 		FInterpCurveVector& SplinePosition = SplineComp->GetSplinePointsPosition();
-		FInterpCurveQuat& SplineRotation = SplineComp->GetSplinePointsRotation();
-		FInterpCurveVector& SplineScale = SplineComp->GetSplinePointsScale();
 
 		const int32 NumPoints = SplinePosition.Points.Num();
 
@@ -1174,17 +1154,15 @@ bool FSplineComponentVisualizer::HandleSnapTo(const bool bInAlign, const bool bI
 	return false;
 }
 
-void FSplineComponentVisualizer::OnSnapToMarkedKey(bool bAlign)
+void FSplineComponentVisualizer::OnSnapToNearestSplinePoint(bool bAlign)
 {
-	const FScopedTransaction Transaction(LOCTEXT("SnapToMarkedSplinePoint", "Snap To Marked Spline Point"));
+	const FScopedTransaction Transaction(LOCTEXT("SnapToNearestSplinePoint", "Snap To Nearest Spline Point"));
 
 	USplineComponent* SplineComp = GetEditedSplineComponent();
 	check(SplineComp != nullptr);
 	check(LastKeyIndexSelected != INDEX_NONE);
 	check(SelectedKeys.Num() == 1);
 	check(SelectedKeys.Contains(LastKeyIndexSelected));
-	check(HasMarkedKey());
-	check(!IsMarkedKey(SplineComp, LastKeyIndexSelected));
 
 	SplineComp->Modify();
 	if (AActor* Owner = SplineComp->GetOwner())
@@ -1192,48 +1170,109 @@ void FSplineComponentVisualizer::OnSnapToMarkedKey(bool bAlign)
 		Owner->Modify();
 	}
 
-	const USplineComponent* MarkedSplineComp = MarkedKeySplineComponent.Get();
-	check(MarkedSplineComp != nullptr);
-
-	const FInterpCurvePoint<FVector>& MarkedPosition = MarkedSplineComp->GetSplinePointsPosition().Points[MarkedKeyIndex];
-	const FInterpCurvePoint<FQuat>& MarkedRotation = MarkedSplineComp->GetSplinePointsRotation().Points[MarkedKeyIndex];
-	const FInterpCurvePoint<FVector>& MarkedScale = MarkedSplineComp->GetSplinePointsScale().Points[MarkedKeyIndex];
-
 	FInterpCurvePoint<FVector>& EditedPosition = SplineComp->GetSplinePointsPosition().Points[LastKeyIndexSelected];
 	FInterpCurvePoint<FQuat>& EditedRotation = SplineComp->GetSplinePointsRotation().Points[LastKeyIndexSelected];
 	FInterpCurvePoint<FVector>& EditedScale = SplineComp->GetSplinePointsScale().Points[LastKeyIndexSelected];
 
+	const FVector WorldPos = SplineComp->GetComponentTransform().TransformPosition(EditedPosition.OutVal); // convert local-space position to world-space
+
+	float NearestDistanceSquared = 0.0f;
+	USplineComponent* NearestSplineComp = nullptr;
+	int32 NearestKeyIndex = INDEX_NONE;
+
+	static const float SnapTol = 5000.0f;
+	float SnapTolSquared = SnapTol * SnapTol;
+
+	// Search all spline components for nearest point.
+	// Only test points in splines whose bounding box contains this point.
+	for (TObjectIterator<USplineComponent> SplineIt; SplineIt; ++SplineIt)
+	{
+		USplineComponent* TestComponent = *SplineIt;
+
+		// Ignore current spline, those that are being destroyed, those with empty bbox.
+		if (TestComponent && TestComponent != SplineComp &&
+			!TestComponent->IsBeingDestroyed() && 
+			!FMath::IsNearlyZero(TestComponent->Bounds.SphereRadius))
+		{
+			FBox TestComponentBoundingBox = TestComponent->Bounds.GetBox().ExpandBy(FVector(SnapTol, SnapTol, SnapTol));
+
+			if (TestComponentBoundingBox.IsInsideOrOn(WorldPos))
+			{
+				const FInterpCurveVector& SplineInfo = TestComponent->GetSplinePointsPosition();
+				const int32 NumPoints = SplineInfo.Points.Num();
+				for (int32 KeyIdx = 0; KeyIdx < NumPoints; KeyIdx++)
+				{
+					const FVector TestKeyWorldPos = TestComponent->GetLocationAtSplinePoint(KeyIdx, ESplineCoordinateSpace::World);
+					float TestDistanceSquared = FVector::DistSquared(TestKeyWorldPos, WorldPos);
+
+					if (TestDistanceSquared < SnapTolSquared && (NearestKeyIndex == INDEX_NONE || TestDistanceSquared < NearestDistanceSquared))
+					{
+						NearestDistanceSquared = TestDistanceSquared;
+						NearestSplineComp = TestComponent;
+						NearestKeyIndex = KeyIdx;
+					}
+				}
+			}
+		}
+	}
+
+	if (!NearestSplineComp || NearestKeyIndex == INDEX_NONE)
+	{
+		UE_LOG(LogSplineComponentVisualizer, Warning, TEXT("No nearest spline point found."));
+		return;
+	}
+
+	const FInterpCurvePoint<FVector>& NearestPosition = NearestSplineComp->GetSplinePointsPosition().Points[NearestKeyIndex];
+	const FInterpCurvePoint<FQuat>& NearestRotation = NearestSplineComp->GetSplinePointsRotation().Points[NearestKeyIndex];
+	const FInterpCurvePoint<FVector>& NearestScale = NearestSplineComp->GetSplinePointsScale().Points[NearestKeyIndex];
+
 	// Copy position
-	const FVector NewWorldPos = MarkedSplineComp->GetComponentTransform().TransformPosition(MarkedPosition.OutVal); // convert local-space position to world-space
+	const FVector NewWorldPos = NearestSplineComp->GetComponentTransform().TransformPosition(NearestPosition.OutVal); // convert local-space position to world-space
 	EditedPosition.OutVal = SplineComp->GetComponentTransform().InverseTransformPosition(NewWorldPos); // convert world-space position to local-space
 
 	if (bAlign)
 	{
 		// Copy tangents
 		EditedPosition.InterpMode = CIM_CurveUser;
-		const FVector NewArriveTangent = MarkedSplineComp->GetComponentTransform().GetRotation().RotateVector(MarkedPosition.ArriveTangent); // convert local-space tangent vectors to world-space
-		const FVector NewLeaveTangent = MarkedSplineComp->GetComponentTransform().GetRotation().RotateVector(MarkedPosition.LeaveTangent); 
-		EditedPosition.ArriveTangent = SplineComp->GetComponentTransform().GetRotation().Inverse().RotateVector(NewArriveTangent); // convert world-space tangent vectors back into local-space
-		EditedPosition.LeaveTangent = SplineComp->GetComponentTransform().GetRotation().Inverse().RotateVector(NewLeaveTangent); 
+		const FVector NewArriveTangent = NearestSplineComp->GetComponentTransform().GetRotation().RotateVector(NearestPosition.ArriveTangent); // convert local-space tangent vectors to world-space
+		const FVector NewLeaveTangent = NearestSplineComp->GetComponentTransform().GetRotation().RotateVector(NearestPosition.LeaveTangent); 
+
+		const FVector ArriveTangent = SplineComp->GetComponentTransform().GetRotation().RotateVector(EditedPosition.ArriveTangent); // convert local-space tangent vectors to world-space
+		const FVector LeaveTangent = SplineComp->GetComponentTransform().GetRotation().RotateVector(EditedPosition.LeaveTangent); 
+
+		// Swap the tangents if they are not pointing in the same general direction
+		float CurrentAngle = FMath::Acos(FVector::DotProduct(ArriveTangent, NewArriveTangent) / (ArriveTangent.Size() * NewArriveTangent.Size()));
+		if (CurrentAngle > HALF_PI)
+		{
+			EditedPosition.ArriveTangent = SplineComp->GetComponentTransform().GetRotation().Inverse().RotateVector(NewLeaveTangent * -1.0f); // convert world-space tangent vectors back into local-space
+			EditedPosition.LeaveTangent = SplineComp->GetComponentTransform().GetRotation().Inverse().RotateVector(NewArriveTangent * -1.0f);
+		}
+		else
+		{
+			EditedPosition.ArriveTangent = SplineComp->GetComponentTransform().GetRotation().Inverse().RotateVector(NewArriveTangent); // convert world-space tangent vectors back into local-space
+			EditedPosition.LeaveTangent = SplineComp->GetComponentTransform().GetRotation().Inverse().RotateVector(NewLeaveTangent);
+		}
   
 		// Copy rotation
-		const FQuat NewRot = MarkedSplineComp->GetComponentTransform().GetRotation() * MarkedRotation.OutVal; // convert local-space rotation to world-space
+		FQuat NewRot = NearestSplineComp->GetComponentTransform().GetRotation() * NearestRotation.OutVal; // convert local-space rotation to world-space
 		EditedRotation.OutVal = SplineComp->GetComponentTransform().GetRotation().Inverse() * NewRot; // convert world-space rotation to local-space
 	}
   
 	// Copy scale - X is not used so ignore it
-	const FVector MarkedSplineCompScale = MarkedSplineComp->GetComponentTransform().GetScale3D();
+	const FVector NearestSplineCompScale = NearestSplineComp->GetComponentTransform().GetScale3D();
 	const FVector SplineCompScale = SplineComp->GetComponentTransform().GetScale3D();
-	float NewScaleY = MarkedSplineCompScale.Y * MarkedScale.OutVal.Y; // convert local-space scale to world-space
-	float NewScaleZ = MarkedSplineCompScale.Z * MarkedScale.OutVal.Z;
+	float NewScaleY = NearestSplineCompScale.Y * NearestScale.OutVal.Y; // convert local-space scale to world-space
+	float NewScaleZ = NearestSplineCompScale.Z * NearestScale.OutVal.Z;
 	EditedScale.OutVal.Y = FMath::IsNearlyZero(SplineCompScale.Y) ? NewScaleY : NewScaleY / SplineCompScale.Y; // convert world-space scale to local-space
 	EditedScale.OutVal.Z = FMath::IsNearlyZero(SplineCompScale.Z) ? NewScaleZ : NewScaleZ / SplineCompScale.Z; 
 
 	// Copy metadata
-	USplineMetadata* SplineMetadata = SplineComp->GetSplinePointsMetadata();
-	if (SplineMetadata)
-	{
-		SplineMetadata->CopyPoint(MarkedSplineComp, MarkedKeyIndex, LastKeyIndexSelected);
+	if (USplineMetadata* SplineMetadata = SplineComp->GetSplinePointsMetadata())
+	{ 
+		if (const USplineMetadata* NearestSplineMetadata = NearestSplineComp->GetSplinePointsMetadata())
+		{ 
+			SplineMetadata->CopyPoint(NearestSplineMetadata, NearestKeyIndex, LastKeyIndexSelected);
+		}
 	}
 
 	SplineComp->UpdateSpline();
@@ -1247,18 +1286,14 @@ void FSplineComponentVisualizer::OnSnapToMarkedKey(bool bAlign)
 	}
 
 	GEditor->RedrawLevelEditingViewports(true);
-
-	OnResetMarkedKey();
 }
 
-bool FSplineComponentVisualizer::CanSnapToMarkedKey() const
+bool FSplineComponentVisualizer::CanSnapToNearestSplinePoint() const
 {
 	USplineComponent* SplineComp = GetEditedSplineComponent();
 	return (SplineComp != nullptr &&
 		SelectedKeys.Num() == 1 &&
-		LastKeyIndexSelected != INDEX_NONE &&
-		HasMarkedKey() &&
-		!IsMarkedKey(SplineComp, LastKeyIndexSelected));
+		LastKeyIndexSelected != INDEX_NONE);
 }
 
 void FSplineComponentVisualizer::OnSnapAll(EAxis::Type InAxis)
@@ -1581,11 +1616,6 @@ void FSplineComponentVisualizer::OnDeleteKey()
 		
 	for (int32 SelectedKeyIndex : SelectedKeysSorted)
 	{
-		if (IsMarkedKey(SplineComp, SelectedKeyIndex))
-		{
-			OnResetMarkedKey();
-		}
-
 		if (SplineMetadata)
 		{
 			SplineMetadata->RemovePoint(SelectedKeyIndex);
@@ -1638,82 +1668,6 @@ bool FSplineComponentVisualizer::IsKeySelectionValid() const
 			LastKeyIndexSelected != INDEX_NONE);
 }
 
-void FSplineComponentVisualizer::OnSetMarkedKey()
-{
-	const FScopedTransaction Transaction(LOCTEXT("SetMarkedSplinePoint", "Set Marked Spline Point"));
-	USplineComponent* SplineComp = GetEditedSplineComponent();
-	check(SplineComp != nullptr);
-	check(LastKeyIndexSelected != INDEX_NONE);
-	check(SelectedKeys.Num() > 0);
-	check(SelectedKeys.Contains(LastKeyIndexSelected));
-
-	if (SelectedKeys.Num() == 1)
-	{
-		MarkedKeyIndex = LastKeyIndexSelected;
-		MarkedKeySplineComponent = SplineComp;
-	}
-}
-
-bool FSplineComponentVisualizer::CanSetMarkedKey() const
-{
-	USplineComponent* SplineComp = GetEditedSplineComponent();
-	return (SplineComp != nullptr &&
-			SelectedKeys.Num() == 1 &&
-			LastKeyIndexSelected != INDEX_NONE);
-}
-
-void FSplineComponentVisualizer::OnResetMarkedKey()
-{
-	MarkedKeyIndex = INDEX_NONE;
-	MarkedKeySplineComponent = nullptr;
-}
-
-bool FSplineComponentVisualizer::CanResetMarkedKey() const
-{
-	return (MarkedKeyIndex != INDEX_NONE && MarkedKeySplineComponent.IsValid());
-}
-
-bool FSplineComponentVisualizer::ContainsMarkedKey(const USplineComponent* InSplineComp) const
-{
-	return (InSplineComp != nullptr &&
-			HasMarkedKey() &&
-			InSplineComp == MarkedKeySplineComponent.Get());
-}
-
-bool FSplineComponentVisualizer::IsMarkedKey(const USplineComponent* InSplineComp, int32 InKeyIdx) const
-{
-	return (ContainsMarkedKey(InSplineComp) && InKeyIdx == MarkedKeyIndex);
-}
-
-bool FSplineComponentVisualizer::HasMarkedKey() const
-{
-	return (MarkedKeyIndex != INDEX_NONE &&
-			MarkedKeySplineComponent.IsValid() &&
-			MarkedKeyIndex < MarkedKeySplineComponent.Get()->GetSplinePointsPosition().Points.Num());
-}
-
-/*
-void FSplineComponentVisualizer::ValidateMarkedKey() 
-{
-	if (MarkedKeySplineComponent.IsValid())
-	{
-		if (MarkedKeyIndex == INDEX_NONE ||
-			MarkedKeyIndex >= MarkedKeySplineComponent.Get()->GetSplinePointsPosition().Points.Num())
-		{
-			MarkedKeySplineComponent = nullptr;
-			MarkedKeyIndex = INDEX_NONE;
-		}
-	}
-	else
-	{
-		if (MarkedKeyIndex != INDEX_NONE)
-		{
-			MarkedKeyIndex = INDEX_NONE;
-		}
-	}
-}
-*/
-
 void FSplineComponentVisualizer::OnLockAxis(EAxis::Type InAxis)
 {
 	AddKeyLockedAxis = InAxis;
@@ -1984,11 +1938,6 @@ TSharedPtr<SWidget> FSplineComponentVisualizer::GenerateContextMenu() const
 				LOCTEXT("KeyTypeTooltip", "Snap align options."),
 				FNewMenuDelegate::CreateSP(this, &FSplineComponentVisualizer::GenerateSnapAlignSubMenu));
 
-			MenuBuilder.AddSubMenu(
-				LOCTEXT("Mark", "Mark"),
-				LOCTEXT("KeyTypeTooltip", "Mark selected spline point."),
-				FNewMenuDelegate::CreateSP(this, &FSplineComponentVisualizer::GenerateMarkSubMenu));
-
 			/* temporarily disabled
 			MenuBuilder.AddSubMenu(
 				LOCTEXT("LockAxis", "Lock Axis"),
@@ -2033,17 +1982,11 @@ void FSplineComponentVisualizer::GenerateSnapAlignSubMenu(FMenuBuilder& MenuBuil
 {
 	MenuBuilder.AddMenuEntry(FLevelEditorCommands::Get().SnapToFloor);
 	MenuBuilder.AddMenuEntry(FLevelEditorCommands::Get().AlignToFloor);
+	MenuBuilder.AddMenuEntry(FSplineComponentVisualizerCommands::Get().SnapToNearestSplinePoint);
+	MenuBuilder.AddMenuEntry(FSplineComponentVisualizerCommands::Get().AlignToNearestSplinePoint);
 	MenuBuilder.AddMenuEntry(FSplineComponentVisualizerCommands::Get().SnapAllToSelectedX);
 	MenuBuilder.AddMenuEntry(FSplineComponentVisualizerCommands::Get().SnapAllToSelectedY);
 	MenuBuilder.AddMenuEntry(FSplineComponentVisualizerCommands::Get().SnapAllToSelectedZ);
-	MenuBuilder.AddMenuEntry(FSplineComponentVisualizerCommands::Get().SnapToMarkedKey);
-	MenuBuilder.AddMenuEntry(FSplineComponentVisualizerCommands::Get().AlignToMarkedKey);
-}
-
-void FSplineComponentVisualizer::GenerateMarkSubMenu(FMenuBuilder& MenuBuilder) const
-{
-	MenuBuilder.AddMenuEntry(FSplineComponentVisualizerCommands::Get().SetMarkedKey);
-	MenuBuilder.AddMenuEntry(FSplineComponentVisualizerCommands::Get().ResetMarkedKey);
 }
 
 void FSplineComponentVisualizer::GenerateLockAxisSubMenu(FMenuBuilder& MenuBuilder) const
diff --git a/Engine/Source/Editor/ComponentVisualizers/Public/SplineComponentVisualizer.h b/Engine/Source/Editor/ComponentVisualizers/Public/SplineComponentVisualizer.h
--- a/Engine/Source/Editor/ComponentVisualizers/Public/SplineComponentVisualizer.h
+++ b/Engine/Source/Editor/ComponentVisualizers/Public/SplineComponentVisualizer.h
@@ -133,14 +133,8 @@ protected:
 	void OnAddKey();
 	bool CanAddKey() const;
 
-	void OnSetMarkedKey();
-	bool CanSetMarkedKey() const;
-
-	void OnResetMarkedKey();
-	bool CanResetMarkedKey() const;
-
-	void OnSnapToMarkedKey(bool bAlign);
-	bool CanSnapToMarkedKey() const;
+	void OnSnapToNearestSplinePoint(bool bAlign);
+	bool CanSnapToNearestSplinePoint() const;
 
 	void OnSnapAll(EAxis::Type InAxis);
 	bool CanSnapAll() const;
@@ -163,16 +157,6 @@ protected:
 	void OnResetToDefault();
 	bool CanResetToDefault() const;
 
-
-	/** Returns true if specified spline component contains the marked key */
-	bool ContainsMarkedKey(const USplineComponent* InSplineComp) const;
-
-	/** Returns true if specified spline component and key index are the marked key */
-	bool IsMarkedKey(const USplineComponent* InSplineComp, int32 InKeyIdx) const;
-
-	/** Returns true if there is a valid marked key */
-	bool HasMarkedKey() const;
-
 	/** Generate the submenu containing the available point types */
 	void GenerateSplinePointTypeSubMenu(FMenuBuilder& MenuBuilder) const;
 
@@ -182,9 +166,6 @@ protected:
 	/** Generate the submenu containing the available snap/align actions */
 	void GenerateSnapAlignSubMenu(FMenuBuilder& MenuBuilder) const;
 	
-	/** Generate the submenu containing the mark key actions */
-	void GenerateMarkSubMenu(FMenuBuilder& MenuBuilder) const;
-
 	/** Generate the submenu containing the lock axis types */
 	void GenerateLockAxisSubMenu(FMenuBuilder& MenuBuilder) const;
 
@@ -231,12 +212,6 @@ protected:
 	/** Whether we currently allow duplication when dragging */
 	bool bAllowDuplication;
 
-	/** Spline component containing marked spline point, used for snapping. */
-	TWeakObjectPtr<USplineComponent> MarkedKeySplineComponent;
-
-	/** Index of marked spline point, used for snapping. Note: this is not necessarily in the current spline component. */
-	int32 MarkedKeyIndex;
-
 	/** Axis to fix when adding new spline points. Uses the value of the currently 
 	    selected spline point's X, Y, or Z value when fix is not equal to none. */
 	EAxis::Type AddKeyLockedAxis;
diff --git a/Engine/Source/Runtime/Engine/Classes/Components/SplineComponent.h b/Engine/Source/Runtime/Engine/Classes/Components/SplineComponent.h
--- a/Engine/Source/Runtime/Engine/Classes/Components/SplineComponent.h
+++ b/Engine/Source/Runtime/Engine/Classes/Components/SplineComponent.h
@@ -47,7 +47,7 @@ public:
 	virtual void AddPoint(float InputKey) PURE_VIRTUAL(USplineMetadata::AddPoint, );
 	virtual void RemovePoint(int32 Index) PURE_VIRTUAL(USplineMetadata::RemovePoint, );
 	virtual void DuplicatePoint(int32 Index) PURE_VIRTUAL(USplineMetadata::DuplicatePoint, );
-	virtual void CopyPoint(const USplineComponent* InFromSplineComp, int32 InFromIndex, int32 InToIndex) PURE_VIRTUAL(USplineMetadata::CopyPoint, );
+	virtual void CopyPoint(const USplineMetadata* FromSplineMetadata, int32 FromIndex, int32 ToIndex) PURE_VIRTUAL(USplineMetadata::CopyPoint, );
 	virtual void Reset(int32 NumPoints) PURE_VIRTUAL(USplineMetadata::Reset, );
 	virtual void Fixup(int32 NumPoints) PURE_VIRTUAL(USplineMetadata::Fixup, );
 };
