Added Select All Spline Points menu option.

diff --git a/Engine/Source/Editor/ComponentVisualizers/Private/SplineComponentVisualizer.cpp b/Engine/Source/Editor/ComponentVisualizers/Private/SplineComponentVisualizer.cpp
--- a/Engine/Source/Editor/ComponentVisualizers/Private/SplineComponentVisualizer.cpp
+++ b/Engine/Source/Editor/ComponentVisualizers/Private/SplineComponentVisualizer.cpp
@@ -47,6 +47,7 @@ public:
 		UI_COMMAND(DeleteKey, "Delete Spline Point", "Delete the currently selected spline point.", EUserInterfaceActionType::Button, FInputChord(EKeys::Delete));
 		UI_COMMAND(DuplicateKey, "Duplicate Spline Point", "Duplicate the currently selected spline point.", EUserInterfaceActionType::Button, FInputChord());
 		UI_COMMAND(AddKey, "Add Spline Point Here", "Add a new spline point at the cursor location.", EUserInterfaceActionType::Button, FInputChord());
+		UI_COMMAND(SelectAll, "Select All Spline Points", "Select all spline points.", EUserInterfaceActionType::Button, FInputChord());
 		UI_COMMAND(ResetToUnclampedTangent, "Unclamped Tangent", "Reset the tangent for this spline point to its default unclamped value.", EUserInterfaceActionType::Button, FInputChord());
 		UI_COMMAND(ResetToClampedTangent, "Clamped Tangent", "Reset the tangent for this spline point to its default clamped value.", EUserInterfaceActionType::Button, FInputChord());
 		UI_COMMAND(SetKeyToCurve, "Curve", "Set spline point to Curve type", EUserInterfaceActionType::RadioButton, FInputChord());
@@ -76,6 +77,9 @@ public:
 	/** Add key */
 	TSharedPtr<FUICommandInfo> AddKey;
 
+	/** Select all */
+	TSharedPtr<FUICommandInfo> SelectAll;
+
 	/** Reset to unclamped tangent */
 	TSharedPtr<FUICommandInfo> ResetToUnclampedTangent;
 
@@ -160,8 +164,13 @@ void FSplineComponentVisualizer::OnRegister()
 
 	SplineComponentVisualizerActions->MapAction(
 		Commands.AddKey,
-		FExecuteAction::CreateSP(this, &FSplineComponentVisualizer::OnAddKey),
-		FCanExecuteAction::CreateSP(this, &FSplineComponentVisualizer::CanAddKey));
+		FExecuteAction::CreateSP(this, &FSplineComponentVisualizer::OnAddKeyToSegment),
+		FCanExecuteAction::CreateSP(this, &FSplineComponentVisualizer::CanAddKeyToSegment));
+
+	SplineComponentVisualizerActions->MapAction(
+		Commands.SelectAll,
+		FExecuteAction::CreateSP(this, &FSplineComponentVisualizer::OnSelectAllSplinePoints),
+		FCanExecuteAction::CreateSP(this, &FSplineComponentVisualizer::CanSelectAllSplinePoints));
 
 	SplineComponentVisualizerActions->MapAction(
 		Commands.ResetToUnclampedTangent,
@@ -1496,8 +1505,7 @@ void FSplineComponentVisualizer::DuplicateKey()
 	GEditor->RedrawLevelEditingViewports(true);
 }
 
-
-bool FSplineComponentVisualizer::CanAddKey() const
+bool FSplineComponentVisualizer::CanAddKeyToSegment() const
 {
 	USplineComponent* SplineComp = GetEditedSplineComponent();
 	if (SplineComp == nullptr)
@@ -1511,7 +1519,7 @@ bool FSplineComponentVisualizer::CanAddKey() const
 	return (SelectedSegmentIndex != INDEX_NONE && SelectedSegmentIndex < NumSegments);
 }
 
-void FSplineComponentVisualizer::OnAddKey()
+void FSplineComponentVisualizer::OnAddKeyToSegment()
 {
 	const FScopedTransaction Transaction(LOCTEXT("AddSplinePoint", "Add Spline Point"));
 	USplineComponent* SplineComp = GetEditedSplineComponent();
@@ -1887,6 +1895,38 @@ bool FSplineComponentVisualizer::CanResetToDefault() const
     }
 }
 
+void FSplineComponentVisualizer::OnSelectAllSplinePoints()
+{
+	USplineComponent* SplineComp = GetEditedSplineComponent();
+	if (SplineComp != nullptr)
+	{
+		const FScopedTransaction Transaction(LOCTEXT("SelectAllSplinePoints", "Select All Spline Points"));
+
+		bool bSelectionChanged = false;
+
+		const FInterpCurveVector& SplineInfo = SplineComp->GetSplinePointsPosition();
+		int32 NumPoints = SplineInfo.Points.Num();
+
+		SelectedKeys.Empty();
+
+		// Spline control point selection always uses transparent box selection.
+		for (int32 KeyIdx = 0; KeyIdx < NumPoints; KeyIdx++)
+		{
+			SelectedKeys.Add(KeyIdx);
+		}
+
+		LastKeyIndexSelected = NumPoints - 1;
+		SelectedSegmentIndex = INDEX_NONE;
+		SelectedTangentHandle = INDEX_NONE;
+		SelectedTangentHandleType = ESelectedTangentHandle::None;
+	}
+}
+
+bool FSplineComponentVisualizer::CanSelectAllSplinePoints() const
+{
+	USplineComponent* SplineComp = GetEditedSplineComponent();
+	return (SplineComp != nullptr);
+}
 
 TSharedPtr<SWidget> FSplineComponentVisualizer::GenerateContextMenu() const
 {
@@ -1902,6 +1942,7 @@ TSharedPtr<SWidget> FSplineComponentVisualizer::GenerateContextMenu() const
 			{
 				MenuBuilder.AddMenuEntry(FSplineComponentVisualizerCommands::Get().DeleteKey);
 				MenuBuilder.AddMenuEntry(FSplineComponentVisualizerCommands::Get().DuplicateKey);
+				MenuBuilder.AddMenuEntry(FSplineComponentVisualizerCommands::Get().SelectAll);
 
 				MenuBuilder.AddSubMenu(
 					LOCTEXT("SplinePointType", "Spline Point Type"),
diff --git a/Engine/Source/Editor/ComponentVisualizers/Public/SplineComponentVisualizer.h b/Engine/Source/Editor/ComponentVisualizers/Public/SplineComponentVisualizer.h
--- a/Engine/Source/Editor/ComponentVisualizers/Public/SplineComponentVisualizer.h
+++ b/Engine/Source/Editor/ComponentVisualizers/Public/SplineComponentVisualizer.h
@@ -130,8 +130,8 @@ protected:
 	void OnDuplicateKey();
 	bool IsKeySelectionValid() const;
 
-	void OnAddKey();
-	bool CanAddKey() const;
+	void OnAddKeyToSegment();
+	bool CanAddKeyToSegment() const;
 
 	void OnSnapToNearestSplinePoint(bool bAlign);
 	bool CanSnapToNearestSplinePoint() const;
@@ -157,6 +157,9 @@ protected:
 	void OnResetToDefault();
 	bool CanResetToDefault() const;
 
+	void OnSelectAllSplinePoints();
+	bool CanSelectAllSplinePoints() const;
+
 	/** Generate the submenu containing the available point types */
 	void GenerateSplinePointTypeSubMenu(FMenuBuilder& MenuBuilder) const;
 
