Spline component - added bi-directional alt-drag, updating tangents properly to avoid loops, interpolate metadata.

diff --git a/Engine/Source/Editor/ComponentVisualizers/Private/SplineComponentVisualizer.cpp b/Engine/Source/Editor/ComponentVisualizers/Private/SplineComponentVisualizer.cpp
--- a/Engine/Source/Editor/ComponentVisualizers/Private/SplineComponentVisualizer.cpp
+++ b/Engine/Source/Editor/ComponentVisualizers/Private/SplineComponentVisualizer.cpp
@@ -139,6 +139,11 @@ FSplineComponentVisualizer::FSplineComponentVisualizer()
 	, SelectedTangentHandle(INDEX_NONE)
 	, SelectedTangentHandleType(ESelectedTangentHandle::None)
 	, bAllowDuplication(true)
+	, bDuplicatingSplineKey(false)
+	, bUpdatingAddSegment(false)
+	, DuplicateDelay(0)
+	, DuplicateDelayAccumulatedDrag(FVector::ZeroVector)
+	, DuplicateCacheSplitSegmentParam(0.0f)
 	, AddKeyLockedAxis(EAxis::None)
 {
 	FSplineComponentVisualizerCommands::Register();
@@ -325,7 +330,6 @@ void FSplineComponentVisualizer::DrawVisualization(const UActorComponent* Compon
 		const float TangentHandleSize = 8.0f;
 
 		// Draw the tangent handles before anything else so they will not overdraw the rest of the spline
-		// Only draw tangent handle when single key is selected
 		if (SplineComp == EditedSplineComp)
 		{
 			for (int32 SelectedKey : SelectedKeys)
@@ -680,6 +684,10 @@ bool FSplineComponentVisualizer::GetWidgetLocation(const FEditorViewportClient*
 			check(SelectedKeys.Contains(LastKeyIndexSelected));
 			const auto& Point = Position.Points[LastKeyIndexSelected];
 			OutLocation = SplineComp->GetComponentTransform().TransformPosition(Point.OutVal);
+			if (!DuplicateDelayAccumulatedDrag.IsZero())
+			{
+				OutLocation += DuplicateDelayAccumulatedDrag;
+			}
 			return true;
 		}
 	}
@@ -720,10 +728,43 @@ bool FSplineComponentVisualizer::HandleInputDelta(FEditorViewportClient* Viewpor
 		{
 			return TransformSelectedTangent(DeltaTranslate);
 		}
+		else if (ViewportClient->IsAltPressed())
+		{
+			if (ViewportClient->GetWidgetMode() == FWidget::WM_Translate && ViewportClient->GetCurrentWidgetAxis() != EAxisList::None && SelectedKeys.Num() == 1)
+			{
+				static const int MaxDuplicationDelay = 3;
+
+				FVector Drag = DeltaTranslate;
+
+				if (bAllowDuplication)
+				{
+					if (DuplicateDelay < MaxDuplicationDelay)
+					{
+						DuplicateDelay++;
+						DuplicateDelayAccumulatedDrag += DeltaTranslate;
+					}
+					else
+					{
+						Drag += DuplicateDelayAccumulatedDrag;
+						DuplicateDelayAccumulatedDrag = FVector::ZeroVector;
+
+						bAllowDuplication = false;
+						bDuplicatingSplineKey = true;
+
+						DuplicateKeyForAltDrag(Drag);
+					}
+				}
+				else
+				{
+					UpdateDuplicateKeyForAltDrag(Drag);
+				}
+
+				return true;
+			}
+		}
 		else
 		{
-			bool bDuplicateKey = ViewportClient->IsAltPressed() && bAllowDuplication;
-			return TransformSelectedKeys(DeltaTranslate, DeltaRotate, DeltaScale, bDuplicateKey);
+			return TransformSelectedKeys(DeltaTranslate, DeltaRotate, DeltaScale);
 		}
 	}
 
@@ -783,7 +824,7 @@ bool FSplineComponentVisualizer::TransformSelectedTangent(const FVector& DeltaTr
 	return false;
 }
 
-bool FSplineComponentVisualizer::TransformSelectedKeys(const FVector& DeltaTranslate, const FRotator& DeltaRotate, const FVector& DeltaScale, bool bDuplicateKey)
+bool FSplineComponentVisualizer::TransformSelectedKeys(const FVector& DeltaTranslate, const FRotator& DeltaRotate, const FVector& DeltaScale)
 {
 	USplineComponent* SplineComp = GetEditedSplineComponent();
 	if (SplineComp != nullptr)
@@ -800,14 +841,6 @@ bool FSplineComponentVisualizer::TransformSelectedKeys(const FVector& DeltaTrans
 
 		SplineComp->Modify();
 
-		if (bDuplicateKey)
-		{
-			DuplicateKey();
-
-			// Don't duplicate again until we release LMB
-			bAllowDuplication = false;
-		}
-
 		for (int32 SelectedKeyIndex : SelectedKeys)
 		{
 			FInterpCurvePoint<FVector>& EditedPoint = SplinePosition.Points[SelectedKeyIndex];
@@ -906,8 +939,8 @@ bool FSplineComponentVisualizer::HandleInputKey(FEditorViewportClient* ViewportC
 			CachedRotation = SplineComp->GetQuaternionAtSplinePoint(Index, ESplineCoordinateSpace::World);
 		}
 
-		// Reset duplication flag on LMB release
-		bAllowDuplication = true;
+		// Reset duplication on LMB release
+		ResetAllowDuplication();
 	}
 
 	if (Event == IE_Pressed)
@@ -1423,18 +1456,6 @@ void FSplineComponentVisualizer::OnDuplicateKey()
 {
 	const FScopedTransaction Transaction(LOCTEXT("DuplicateSplinePoint", "Duplicate Spline Point"));
 	
-	USplineComponent* SplineComp = GetEditedSplineComponent();
-	DuplicateKey();
-
-	SplineComp->UpdateSpline();
-	SplineComp->bSplineHasBeenEdited = true;
-
-	NotifyPropertyModified(SplineComp, SplineCurvesProperty);
-}
-
-
-void FSplineComponentVisualizer::DuplicateKey()
-{
 	USplineComponent* SplineComp = GetEditedSplineComponent();
 	check(SplineComp != nullptr);
 	check(LastKeyIndexSelected != INDEX_NONE);
@@ -1502,6 +1523,16 @@ void FSplineComponentVisualizer::DuplicateKey()
 	SelectedTangentHandle = INDEX_NONE;
 	SelectedTangentHandleType = ESelectedTangentHandle::None;
 
+	SplineComp->UpdateSpline();
+	SplineComp->bSplineHasBeenEdited = true;
+
+	NotifyPropertyModified(SplineComp, SplineCurvesProperty);
+
+	if (SelectedKeys.Num() == 1)
+	{
+		CachedRotation = SplineComp->GetQuaternionAtSplinePoint(LastKeyIndexSelected, ESplineCoordinateSpace::World);
+	}
+
 	GEditor->RedrawLevelEditingViewports(true);
 }
 
@@ -1530,68 +1561,448 @@ void FSplineComponentVisualizer::OnAddKeyToSegment()
 	check(SelectedTangentHandle == INDEX_NONE);
 	check(SelectedTangentHandleType == ESelectedTangentHandle::None);
 
+	SplitSegment(SelectedSplinePosition, SelectedSegmentIndex);
+
+	CachedRotation = SplineComp->GetQuaternionAtSplinePoint(LastKeyIndexSelected, ESplineCoordinateSpace::World);
+
+	SelectedSplinePosition = FVector::ZeroVector;
+	SelectedSegmentIndex = INDEX_NONE;
+}
+
+bool FSplineComponentVisualizer::DuplicateKeyForAltDrag(const FVector& InDrag)
+{
+	USplineComponent* SplineComp = GetEditedSplineComponent();
+	check(SplineComp != nullptr);
+	check(LastKeyIndexSelected != INDEX_NONE);
+	check(SelectedKeys.Num() == 1);
+	check(SelectedKeys.Contains(LastKeyIndexSelected));
+
+	// Insert duplicates into the list, highest index first, so that the lower indices remain the same
+	FInterpCurveVector& SplinePosition = SplineComp->GetSplinePointsPosition();
+
+	// Find key position in world space
+	int32 CurrentIndex = LastKeyIndexSelected;
+	const FVector CurrentKeyWorldPos = SplineComp->GetComponentTransform().TransformPosition(SplinePosition.Points[CurrentIndex].OutVal);
+
+	// Determine direction to insert new point				
+	bool bHasPrevKey = false;
+	float PrevAngle = 0.0f;
+	if (CurrentIndex > 0)
+	{
+		FVector PrevKeyWorldPos = SplineComp->GetComponentTransform().TransformPosition(SplinePosition.Points[CurrentIndex - 1].OutVal);
+		FVector SegmentDirection = PrevKeyWorldPos - CurrentKeyWorldPos;
+		if (!SegmentDirection.IsZero())
+		{
+			PrevAngle = FMath::Acos(FVector::DotProduct(InDrag, SegmentDirection) / (InDrag.Size() * SegmentDirection.Size()));
+		}
+		else
+		{
+			PrevAngle = HALF_PI;
+		}
+		bHasPrevKey = true;
+	}
+
+	bool bHasNextKey = false;
+	float NextAngle = 0.0f;
+	if (CurrentIndex + 1 < SplinePosition.Points.Num())
+	{
+		FVector NextKeyWorldPos = SplineComp->GetComponentTransform().TransformPosition(SplinePosition.Points[CurrentIndex + 1].OutVal);
+		FVector SegmentDirection = NextKeyWorldPos - CurrentKeyWorldPos;
+		if (!SegmentDirection.IsZero())
+		{
+			NextAngle = FMath::Acos(FVector::DotProduct(InDrag, SegmentDirection) / (InDrag.Size() * SegmentDirection.Size()));
+		}
+		else
+		{
+			NextAngle = HALF_PI;
+		}
+		bHasNextKey = true;
+	}
+
+	// Set key index to which the drag will be applied after duplication
+	int32 SegmentIndex = CurrentIndex;
+
+	if ((bHasPrevKey && bHasNextKey && PrevAngle < NextAngle) ||
+		(bHasPrevKey && !bHasNextKey && PrevAngle < HALF_PI) ||
+		(!bHasPrevKey && bHasNextKey && NextAngle >= HALF_PI))
+	{
+		SegmentIndex--;
+	}
+
+	FVector WorldPos = CurrentKeyWorldPos + InDrag;
+
+	// Split existing segment or add new segment
+	if (SegmentIndex >= 0 && SegmentIndex < SplinePosition.Points.Num() - 1)
+	{
+		SplitSegment(WorldPos, SegmentIndex);
+	}
+	else
+	{
+		AddSegment(WorldPos, (SegmentIndex > 0));
+		bUpdatingAddSegment = true;
+	}
+
+	// Unset tangent handle selection
+	SelectedTangentHandle = INDEX_NONE;
+	SelectedTangentHandleType = ESelectedTangentHandle::None;
+	
+	return true;
+}
+
+bool FSplineComponentVisualizer::UpdateDuplicateKeyForAltDrag(const FVector& InDrag)
+{
+	if (bUpdatingAddSegment)
+	{
+		UpdateAddSegment(InDrag);
+	}
+	else
+	{
+		UpdateSplitSegment(InDrag);
+	}
+
+	return true;
+}
+
+float FSplineComponentVisualizer::FindNearest(const FVector& InLocalPos, int32 InSegmentIndex, FVector& OutSplinePos, FVector& OutSplineTangent) const
+{
+	USplineComponent* SplineComp = GetEditedSplineComponent();
+	check(SplineComp != nullptr);
+	check(InSegmentIndex >= 0);
+	check(InSegmentIndex < SplineComp->GetSplinePointsPosition().Points.Num() - 1);
+
+	FInterpCurveVector& SplinePosition = SplineComp->GetSplinePointsPosition();
+	float OutSquaredDistance = 0.0f;
+	float t = SplinePosition.InaccurateFindNearestOnSegment(InLocalPos, InSegmentIndex, OutSquaredDistance);
+	OutSplinePos = SplinePosition.Eval(t, FVector::ZeroVector);
+	OutSplineTangent = SplinePosition.EvalDerivative(t, FVector::ZeroVector);
+
+	return t;
+}
+
+void FSplineComponentVisualizer::SplitSegment(const FVector& InWorldPos, int32 InSegmentIndex)
+{
+	USplineComponent* SplineComp = GetEditedSplineComponent();
+	check(SplineComp != nullptr);
+	check(LastKeyIndexSelected != INDEX_NONE);
+	check(SelectedKeys.Num() == 1);
+	check(SelectedKeys.Contains(LastKeyIndexSelected));
+	check(InSegmentIndex >= 0);
+	check(InSegmentIndex < SplineComp->GetSplinePointsPosition().Points.Num() - 1);
+
 	SplineComp->Modify();
 	if (AActor* Owner = SplineComp->GetOwner())
 	{
 		Owner->Modify();
 	}
 
+	// Compute local pos
+	FVector LocalPos = SplineComp->GetComponentTransform().InverseTransformPosition(InWorldPos);
+
+	FVector SplinePos, SplineTangent;
+	float SplineParam = FindNearest(LocalPos, InSegmentIndex, SplinePos, SplineTangent);
+	float t = SplineParam - static_cast<float>(InSegmentIndex);
+
+	if (bDuplicatingSplineKey)
+	{
+		DuplicateCacheSplitSegmentParam = t;
+	}
+
+	int32 SegmentBeginIndex = InSegmentIndex;
+	int32 SegmentSplitIndex = InSegmentIndex + 1;
+
 	FInterpCurveVector& SplinePosition = SplineComp->GetSplinePointsPosition();
 	FInterpCurveQuat& SplineRotation = SplineComp->GetSplinePointsRotation();
 	FInterpCurveVector& SplineScale = SplineComp->GetSplinePointsScale();
 	USplineMetadata* SplineMetadata = SplineComp->GetSplinePointsMetadata();
 	
+	// Set adjacent points to CurveAuto so their tangents adjust automatically as new point moves.
+	if (SplinePosition.Points[SegmentBeginIndex].InterpMode == CIM_CurveUser)
+	{
+		SplinePosition.Points[SegmentBeginIndex].InterpMode = CIM_CurveAuto;
+	}
+	if (SplinePosition.Points[SegmentSplitIndex].InterpMode == CIM_CurveUser)
+	{
+		SplinePosition.Points[SegmentSplitIndex].InterpMode = CIM_CurveAuto;
+	}
+
+	// Compute interpolated scale
+	FVector NewScale;
+	FInterpCurvePoint<FVector>& PrevScale = SplineScale.Points[SegmentBeginIndex];
+	FInterpCurvePoint<FVector>& NextScale = SplineScale.Points[SegmentSplitIndex];
+	NewScale = FMath::LerpStable(PrevScale.OutVal, NextScale.OutVal, t);
+
+	// Compute interpolated rot
+	FQuat NewRot;
+	FInterpCurvePoint<FQuat>& PrevRot = SplineRotation.Points[SegmentBeginIndex];
+	FInterpCurvePoint<FQuat>& NextRot = SplineRotation.Points[SegmentSplitIndex];
+	NewRot = FMath::Lerp(PrevRot.OutVal, NextRot.OutVal, t);
+
 	FInterpCurvePoint<FVector> NewPoint(
-		SelectedSegmentIndex,
-		SplineComp->GetComponentTransform().InverseTransformPosition(SelectedSplinePosition),
+		SegmentSplitIndex,
+		SplinePos,
 		FVector::ZeroVector,
 		FVector::ZeroVector,
 		CIM_CurveAuto);
 
 	FInterpCurvePoint<FQuat> NewRotPoint(
-		SelectedSegmentIndex,
-		FQuat::Identity,
+		SegmentSplitIndex,
+		NewRot,
 		FQuat::Identity,
 		FQuat::Identity,
 		CIM_CurveAuto);
 
 	FInterpCurvePoint<FVector> NewScalePoint(
-		SelectedSegmentIndex,
-		FVector(1.0f),
+		SegmentSplitIndex,
+		NewScale,
 		FVector::ZeroVector,
 		FVector::ZeroVector,
 		CIM_CurveAuto);
 
-	SplinePosition.Points.Insert(NewPoint, SelectedSegmentIndex + 1);
-	SplineRotation.Points.Insert(NewRotPoint, SelectedSegmentIndex + 1);
-	SplineScale.Points.Insert(NewScalePoint, SelectedSegmentIndex + 1);
+	SplinePosition.Points.Insert(NewPoint, SegmentSplitIndex);
+	SplineRotation.Points.Insert(NewRotPoint, SegmentSplitIndex);
+	SplineScale.Points.Insert(NewScalePoint, SegmentSplitIndex);
 	if (SplineMetadata)
 	{
-		SplineMetadata->InsertPoint(SelectedSegmentIndex, SelectedSegmentIndex + 1);
+		SplineMetadata->DuplicatePoint(SegmentSplitIndex);
 	}
 
 	// Adjust input keys of subsequent points
-	for (int Index = SelectedSegmentIndex + 1; Index < SplinePosition.Points.Num(); Index++)
+	for (int Index = SegmentSplitIndex + 1; Index < SplinePosition.Points.Num(); Index++)
 	{
 		SplinePosition.Points[Index].InVal += 1.0f;
 		SplineRotation.Points[Index].InVal += 1.0f;
 		SplineScale.Points[Index].InVal += 1.0f;
 	}
 
-	// Set selection to 'next' key
-	ChangeSelectionState(SelectedSegmentIndex + 1, false);
-	SelectedSegmentIndex = INDEX_NONE;
+	// Set selection to new key
+	ChangeSelectionState(SegmentSplitIndex, false);
 
 	SplineComp->UpdateSpline();
 	SplineComp->bSplineHasBeenEdited = true;
 
 	NotifyPropertyModified(SplineComp, SplineCurvesProperty);
 
-	CachedRotation = SplineComp->GetQuaternionAtSplinePoint(LastKeyIndexSelected, ESplineCoordinateSpace::World);
+	GEditor->RedrawLevelEditingViewports(true);
+}
+
+void FSplineComponentVisualizer::UpdateSplitSegment(const FVector& InDrag)
+{
+	USplineComponent* SplineComp = GetEditedSplineComponent();
+	check(SplineComp != nullptr);
+	check(LastKeyIndexSelected != INDEX_NONE);
+	check(SelectedKeys.Num() == 1);
+	check(SelectedKeys.Contains(LastKeyIndexSelected));
+	check(LastKeyIndexSelected > 0);
+	check(LastKeyIndexSelected < SplineComp->GetSplinePointsPosition().Points.Num() - 1);
+
+	SplineComp->Modify();
+	if (AActor* Owner = SplineComp->GetOwner())
+	{
+		Owner->Modify();
+	}
+
+	int32 SegmentStartIndex = LastKeyIndexSelected - 1;
+	int32 SegmentSplitIndex = LastKeyIndexSelected;
+	int32 SegmentEndIndex = LastKeyIndexSelected + 1;
+
+	FInterpCurveVector& SplinePosition = SplineComp->GetSplinePointsPosition();
+	FInterpCurveVector& SplineScale = SplineComp->GetSplinePointsScale();
+	FInterpCurveQuat& SplineRotation = SplineComp->GetSplinePointsRotation();
+	USplineMetadata* SplineMetadata = SplineComp->GetSplinePointsMetadata();
+
+	// Find key position in world space
+	FInterpCurvePoint<FVector>& EditedPoint = SplinePosition.Points[SegmentSplitIndex];
+	const FVector CurrentWorldPos = SplineComp->GetComponentTransform().TransformPosition(EditedPoint.OutVal);
+
+	// Move in world space
+	const FVector NewWorldPos = CurrentWorldPos + InDrag;
+
+	// Convert back to local space
+	FVector LocalPos = SplineComp->GetComponentTransform().InverseTransformPosition(NewWorldPos);
+
+	FVector SplinePos0, SplinePos1;
+	FVector SplineTangent0, SplineTangent1;
+	float t = 0.0f;
+	float SplineParam0 = FindNearest(LocalPos, SegmentStartIndex, SplinePos0, SplineTangent0);
+	float t0 = SplineParam0 - static_cast<float>(SegmentStartIndex);
+	float SplineParam1 = FindNearest(LocalPos, SegmentSplitIndex, SplinePos1, SplineTangent1);
+	float t1 = SplineParam1 - static_cast<float>(SegmentSplitIndex);
+
+	// Calculate params
+	if (FVector::Distance(LocalPos, SplinePos0) < FVector::Distance(LocalPos, SplinePos1))
+	{
+		t = DuplicateCacheSplitSegmentParam * t0;
+	}
+	else
+	{
+		t = DuplicateCacheSplitSegmentParam + (1 - DuplicateCacheSplitSegmentParam) * t1;
+	}
+	DuplicateCacheSplitSegmentParam = t;
+
+	// Update location
+	EditedPoint.OutVal = LocalPos;
+
+	// Update scale
+	FInterpCurvePoint<FVector>& EditedScale = SplineScale.Points[SegmentSplitIndex];
+	FInterpCurvePoint<FVector>& PrevScale = SplineScale.Points[SegmentStartIndex];
+	FInterpCurvePoint<FVector>& NextScale = SplineScale.Points[SegmentEndIndex];
+	EditedScale.OutVal = FMath::LerpStable(PrevScale.OutVal, NextScale.OutVal, t);
+
+	// Update rot
+	FInterpCurvePoint<FQuat>& EditedRot = SplineRotation.Points[SegmentSplitIndex];
+	FInterpCurvePoint<FQuat>& PrevRot = SplineRotation.Points[SegmentStartIndex];
+	FInterpCurvePoint<FQuat>& NextRot = SplineRotation.Points[SegmentEndIndex];
+	EditedRot.OutVal = FMath::Lerp(PrevRot.OutVal, NextRot.OutVal, t);
+
+	// Update metadata
+	if (SplineMetadata)
+	{
+		SplineMetadata->UpdatePoint(SegmentSplitIndex, t);
+	}
+
+	SplineComp->UpdateSpline();
+	SplineComp->bSplineHasBeenEdited = true;
+
+	NotifyPropertyModified(SplineComp, SplineCurvesProperty);
 
 	GEditor->RedrawLevelEditingViewports(true);
 }
 
+void FSplineComponentVisualizer::AddSegment(const FVector& InWorldPos, bool bAppend)
+{
+	USplineComponent* SplineComp = GetEditedSplineComponent();
+	check(SplineComp != nullptr);
+
+	SplineComp->Modify();
+	if (AActor* Owner = SplineComp->GetOwner())
+	{
+		Owner->Modify();
+	}
+
+	int32 KeyIdx = 0;
+	int32 NewKeyIdx = 0;
+
+	FInterpCurveVector& SplinePosition = SplineComp->GetSplinePointsPosition();
+
+	if (bAppend)
+	{
+		NewKeyIdx = SplinePosition.Points.Num();
+		KeyIdx = NewKeyIdx - 1;
+	}
+
+	FInterpCurveQuat& SplineRotation = SplineComp->GetSplinePointsRotation();
+	FInterpCurveVector& SplineScale = SplineComp->GetSplinePointsScale();
+	USplineMetadata* SplineMetadata = SplineComp->GetSplinePointsMetadata();
+	
+	// Set adjacent point to CurveAuto so its tangent adjusts automatically as new point moves.
+	if (SplinePosition.Points[KeyIdx].InterpMode == CIM_CurveUser)
+	{
+		SplinePosition.Points[KeyIdx].InterpMode = CIM_CurveAuto;
+	}
+
+	// Compute local pos
+	FVector LocalPos = SplineComp->GetComponentTransform().InverseTransformPosition(InWorldPos);
+
+	FInterpCurvePoint<FVector> NewPoint(
+		NewKeyIdx,
+		LocalPos,
+		FVector::ZeroVector,
+		FVector::ZeroVector,
+		SplinePosition.Points[KeyIdx].InterpMode);
+
+	FInterpCurvePoint<FQuat> NewRotPoint(
+		NewKeyIdx,
+		SplineRotation.Points[KeyIdx].OutVal,
+		FQuat::Identity,
+		FQuat::Identity,
+		CIM_CurveAuto);
+
+	FInterpCurvePoint<FVector> NewScalePoint(
+		NewKeyIdx,
+		SplineScale.Points[KeyIdx].OutVal,
+		FVector::ZeroVector,
+		FVector::ZeroVector,
+		CIM_CurveAuto);
+
+	if (KeyIdx == 0)
+	{
+		SplinePosition.Points.Insert(NewPoint, KeyIdx);
+		SplineRotation.Points.Insert(NewRotPoint, KeyIdx);
+		SplineScale.Points.Insert(NewScalePoint, KeyIdx);
+	}
+	else
+	{
+		SplinePosition.Points.Emplace(NewPoint);
+		SplineRotation.Points.Emplace(NewRotPoint);
+		SplineScale.Points.Emplace(NewScalePoint);
+	}
+
+	// Adjust input keys of subsequent points
+	if (!bAppend)
+	{
+		for (int Index = 1; Index < SplinePosition.Points.Num(); Index++)
+		{
+			SplinePosition.Points[Index].InVal += 1.0f;
+			SplineRotation.Points[Index].InVal += 1.0f;
+			SplineScale.Points[Index].InVal += 1.0f;
+		}
+	}
+
+	if (SplineMetadata)
+	{
+		SplineMetadata->DuplicatePoint(KeyIdx);
+	}
+
+	// Set selection to key
+	ChangeSelectionState(NewKeyIdx, false);
+
+	SplineComp->UpdateSpline();
+	SplineComp->bSplineHasBeenEdited = true;
+
+	NotifyPropertyModified(SplineComp, SplineCurvesProperty);
+
+	GEditor->RedrawLevelEditingViewports(true);
+}
+
+void FSplineComponentVisualizer::UpdateAddSegment(const FVector& InDrag)
+{
+	USplineComponent* SplineComp = GetEditedSplineComponent();
+	check(SplineComp != nullptr);
+	check(LastKeyIndexSelected != INDEX_NONE);
+	check(SelectedKeys.Num() == 1);
+	check(SelectedKeys.Contains(LastKeyIndexSelected));
+	// Only work on keys at either end of the spline 
+	check(LastKeyIndexSelected == 0 || LastKeyIndexSelected == SplineComp->GetSplinePointsPosition().Points.Num() - 1);
+
+	SplineComp->Modify();
+	if (AActor* Owner = SplineComp->GetOwner())
+	{
+		Owner->Modify();
+	}
+
+	// Move added point to new position
+	FInterpCurvePoint<FVector>& AddedPoint = SplineComp->GetSplinePointsPosition().Points[LastKeyIndexSelected];
+	const FVector CurrentWorldPos = SplineComp->GetComponentTransform().TransformPosition(AddedPoint.OutVal);
+	const FVector NewWorldPos = CurrentWorldPos + InDrag;
+	AddedPoint.OutVal = SplineComp->GetComponentTransform().InverseTransformPosition(NewWorldPos);
+
+	SplineComp->UpdateSpline();
+	SplineComp->bSplineHasBeenEdited = true;
+
+	NotifyPropertyModified(SplineComp, SplineCurvesProperty);
+
+	GEditor->RedrawLevelEditingViewports(true);
+}
+
+void FSplineComponentVisualizer::ResetAllowDuplication()
+{
+	bAllowDuplication = true;
+	bDuplicatingSplineKey = false;
+	bUpdatingAddSegment = false;
+	DuplicateDelay = 0;
+	DuplicateDelayAccumulatedDrag = FVector::ZeroVector;
+	DuplicateCacheSplitSegmentParam = 0.0f;
+}
 
 void FSplineComponentVisualizer::OnDeleteKey()
 {
diff --git a/Engine/Source/Editor/ComponentVisualizers/Public/SplineComponentVisualizer.h b/Engine/Source/Editor/ComponentVisualizers/Public/SplineComponentVisualizer.h
--- a/Engine/Source/Editor/ComponentVisualizers/Public/SplineComponentVisualizer.h
+++ b/Engine/Source/Editor/ComponentVisualizers/Public/SplineComponentVisualizer.h
@@ -112,21 +112,39 @@ protected:
 	bool TransformSelectedTangent(const FVector& DeltaTranslate);
 
 	/** Transforms selected tangent by given translate, rotate and scale */
-	bool TransformSelectedKeys(const FVector& DeltaTranslate, const FRotator& DeltaRotate = FRotator::ZeroRotator, const FVector& DeltaScale = FVector::ZeroVector, bool bDuplicateKey = false);
-
-	/** Snaps edited point and rot point to given world location, up vector, and tangent */
-	void SnapTo(USplineComponent *SplineComp, int32 KeyIdx,
-		const FVector& InLocation, const FVector& InUpVector = FVector::ZeroVector, bool bInAlignUpVector = false, const FVector& InTangentVector = FVector::ZeroVector, bool bInAlignTangentVector = false);
+	bool TransformSelectedKeys(const FVector& DeltaTranslate, const FRotator& DeltaRotate = FRotator::ZeroRotator, const FVector& DeltaScale = FVector::ZeroVector);
 
 	/** Update the key selection state of the visualizer */
 	void ChangeSelectionState(int32 Index, bool bIsCtrlHeld);
 
-	/** Duplicates the selected spline key(s) */
-	void DuplicateKey();
+	/** Alt-drag: duplicates the selected spline key */
+	virtual bool DuplicateKeyForAltDrag(const FVector& InDrag);
+
+	/** Alt-drag: updates duplicated selected spline key */
+	virtual bool UpdateDuplicateKeyForAltDrag(const FVector& InDrag);
+
+	/** Return spline data for point on spline closest to input point */
+	virtual float FindNearest(const FVector& InLocalPos, int32 InSegmentStartIndex, FVector& OutSplinePos, FVector& OutSplineTangent) const;
+
+	/** Split segment using given world position */
+	virtual void SplitSegment(const FVector& InWorldPos, int32 InSegmentIndex);
+
+	/** Update split segment based on drag offset */
+	virtual void UpdateSplitSegment(const FVector& InDrag);
+
+	/** Add segment to beginning or end of spline */
+	virtual void AddSegment(const FVector& InWorldPos, bool bAppend);
+
+	/** Add segment to beginning or end of spline */
+	virtual void UpdateAddSegment(const FVector& InWorldPos);
+
+	/** Alt-drag: duplicates the selected spline key */
+	virtual void ResetAllowDuplication();
 
 	void OnDeleteKey();
 	bool CanDeleteKey() const;
 
+	/** Duplicates selected spline keys in place */
 	void OnDuplicateKey();
 	bool IsKeySelectionValid() const;
 
@@ -215,6 +233,21 @@ protected:
 	/** Whether we currently allow duplication when dragging */
 	bool bAllowDuplication;
 
+	/** Alt-drag: True when in process of duplicating a spline key. */
+	bool bDuplicatingSplineKey;
+
+	/** Alt-drag: True when in process of adding end segment. */
+	bool bUpdatingAddSegment;
+
+	/** Alt-drag: Delays duplicating control point to accumulate sufficient drag input offset. */
+	uint32 DuplicateDelay;
+
+	/** Alt-drag: Accumulates delayed drag offset. */
+	FVector DuplicateDelayAccumulatedDrag;
+
+	/** Alt-drag: Cached segment parameter for split segment at new control point */
+	float DuplicateCacheSplitSegmentParam;
+
 	/** Axis to fix when adding new spline points. Uses the value of the currently 
 	    selected spline point's X, Y, or Z value when fix is not equal to none. */
 	EAxis::Type AddKeyLockedAxis;
diff --git a/Engine/Source/Runtime/Engine/Classes/Components/SplineComponent.h b/Engine/Source/Runtime/Engine/Classes/Components/SplineComponent.h
--- a/Engine/Source/Runtime/Engine/Classes/Components/SplineComponent.h
+++ b/Engine/Source/Runtime/Engine/Classes/Components/SplineComponent.h
@@ -43,7 +43,10 @@ class ENGINE_API USplineMetadata : public UObject
 	GENERATED_UCLASS_BODY()
 
 public:
-	virtual void InsertPoint(float InputKey, int32 Index) PURE_VIRTUAL(USplineMetadata::InsertPoint, );
+	/** Insert point before index, lerping metadata between previous and next key values */
+	virtual void InsertPoint(int32 Index, float t) PURE_VIRTUAL(USplineMetadata::InsertPoint, );
+	/** Update point at index by lerping metadata between previous and next key values */
+	virtual void UpdatePoint(int32 Index, float t) PURE_VIRTUAL(USplineMetadata::UpdatePoint, );
 	virtual void AddPoint(float InputKey) PURE_VIRTUAL(USplineMetadata::AddPoint, );
 	virtual void RemovePoint(int32 Index) PURE_VIRTUAL(USplineMetadata::RemovePoint, );
 	virtual void DuplicatePoint(int32 Index) PURE_VIRTUAL(USplineMetadata::DuplicatePoint, );
