Added spline point "Align Perpendicular to Nearest Point" menu option.

diff --git a/Engine/Source/Editor/ComponentVisualizers/Private/SplineComponentVisualizer.cpp b/Engine/Source/Editor/ComponentVisualizers/Private/SplineComponentVisualizer.cpp
--- a/Engine/Source/Editor/ComponentVisualizers/Private/SplineComponentVisualizer.cpp
+++ b/Engine/Source/Editor/ComponentVisualizers/Private/SplineComponentVisualizer.cpp
@@ -55,6 +55,7 @@ public:
 		UI_COMMAND(SetKeyToConstant, "Constant", "Set spline point to Constant type", EUserInterfaceActionType::RadioButton, FInputChord());
 		UI_COMMAND(SnapToNearestSplinePoint, "Snap to Nearest Spline Point", "Snap to nearest spline point.", EUserInterfaceActionType::Button, FInputChord());
 		UI_COMMAND(AlignToNearestSplinePoint, "Align to Nearest Spline Point", "Align to nearest spline point.", EUserInterfaceActionType::Button, FInputChord());
+		UI_COMMAND(AlignPerpendicularToNearestSplinePoint, "Align Perpendicular to Nearest Spline Point", "Align perpendicular to nearest spline point.", EUserInterfaceActionType::Button, FInputChord());
 		UI_COMMAND(SnapAllToSelectedX, "Snap All To Selected X", "Snap all spline points to selected spline point X.", EUserInterfaceActionType::Button, FInputChord());
 		UI_COMMAND(SnapAllToSelectedY, "Snap All To Selected Y", "Snap all spline points to selected spline point Y.", EUserInterfaceActionType::Button, FInputChord());
 		UI_COMMAND(SnapAllToSelectedZ, "Snap All To Selected Z", "Snap all spline points to selected spline point Z.", EUserInterfaceActionType::Button, FInputChord());
@@ -101,6 +102,9 @@ public:
 	/** Align to nearest spline point on another spline component */
 	TSharedPtr<FUICommandInfo> AlignToNearestSplinePoint;
 
+	/** Align perpendicular to nearest spline point on another spline component */
+	TSharedPtr<FUICommandInfo> AlignPerpendicularToNearestSplinePoint;
+
 	/** Snap all spline points to selected point X */
 	TSharedPtr<FUICommandInfo> SnapAllToSelectedX;
 
@@ -207,12 +211,17 @@ void FSplineComponentVisualizer::OnRegister()
 
 	SplineComponentVisualizerActions->MapAction(
 		Commands.SnapToNearestSplinePoint,
-		FExecuteAction::CreateSP(this, &FSplineComponentVisualizer::OnSnapToNearestSplinePoint, false),
+		FExecuteAction::CreateSP(this, &FSplineComponentVisualizer::OnSnapToNearestSplinePoint, ESplineComponentSnapMode::Snap),
 		FCanExecuteAction::CreateSP(this, &FSplineComponentVisualizer::CanSnapToNearestSplinePoint));
 
 	SplineComponentVisualizerActions->MapAction(
 		Commands.AlignToNearestSplinePoint,
-		FExecuteAction::CreateSP(this, &FSplineComponentVisualizer::OnSnapToNearestSplinePoint, true),
+		FExecuteAction::CreateSP(this, &FSplineComponentVisualizer::OnSnapToNearestSplinePoint, ESplineComponentSnapMode::AlignToTangent),
+		FCanExecuteAction::CreateSP(this, &FSplineComponentVisualizer::CanSnapToNearestSplinePoint));
+
+	SplineComponentVisualizerActions->MapAction(
+		Commands.AlignPerpendicularToNearestSplinePoint,
+		FExecuteAction::CreateSP(this, &FSplineComponentVisualizer::OnSnapToNearestSplinePoint, ESplineComponentSnapMode::AlignPerpendicularToTangent),
 		FCanExecuteAction::CreateSP(this, &FSplineComponentVisualizer::CanSnapToNearestSplinePoint));
 
 	SplineComponentVisualizerActions->MapAction(
@@ -1196,7 +1205,7 @@ bool FSplineComponentVisualizer::HandleSnapTo(const bool bInAlign, const bool bI
 	return false;
 }
 
-void FSplineComponentVisualizer::OnSnapToNearestSplinePoint(bool bAlign)
+void FSplineComponentVisualizer::OnSnapToNearestSplinePoint(ESplineComponentSnapMode::Type InSnapMode)
 {
 	const FScopedTransaction Transaction(LOCTEXT("SnapToNearestSplinePoint", "Snap To Nearest Spline Point"));
 
@@ -1272,48 +1281,66 @@ void FSplineComponentVisualizer::OnSnapToNearestSplinePoint(bool bAlign)
 	const FVector NewWorldPos = NearestSplineComp->GetComponentTransform().TransformPosition(NearestPosition.OutVal); // convert local-space position to world-space
 	EditedPosition.OutVal = SplineComp->GetComponentTransform().InverseTransformPosition(NewWorldPos); // convert world-space position to local-space
 
-	if (bAlign)
+	if (InSnapMode == ESplineComponentSnapMode::AlignToTangent || InSnapMode == ESplineComponentSnapMode::AlignPerpendicularToTangent)
 	{
 		// Copy tangents
-		EditedPosition.InterpMode = CIM_CurveUser;
-		const FVector NewArriveTangent = NearestSplineComp->GetComponentTransform().GetRotation().RotateVector(NearestPosition.ArriveTangent); // convert local-space tangent vectors to world-space
-		const FVector NewLeaveTangent = NearestSplineComp->GetComponentTransform().GetRotation().RotateVector(NearestPosition.LeaveTangent); 
+		FVector AlignTangent;
+		FQuat AlignRot;
+
+		// Copy tangents
+		FVector NewTangent;
+		if (InSnapMode == ESplineComponentSnapMode::AlignPerpendicularToTangent)
+		{
+			// Rotate tangent by 90 degrees
+			const FVector UpVector = NearestSplineComp->GetUpVectorAtSplinePoint(NearestKeyIndex, ESplineCoordinateSpace::Local);
+			const FQuat DeltaRotate(UpVector, HALF_PI);
+			NewTangent = DeltaRotate.RotateVector(NearestPosition.ArriveTangent); // apply local-space rotation
+			NewTangent = NearestSplineComp->GetComponentTransform().GetRotation().RotateVector(NewTangent); // convert local-space tangent vectors to world-space
+		}
+		else
+		{
+			NewTangent = NearestSplineComp->GetComponentTransform().GetRotation().RotateVector(NearestPosition.ArriveTangent); // convert local-space tangent vectors to world-space
+		}
 
-		const FVector ArriveTangent = SplineComp->GetComponentTransform().GetRotation().RotateVector(EditedPosition.ArriveTangent); // convert local-space tangent vectors to world-space
-		const FVector LeaveTangent = SplineComp->GetComponentTransform().GetRotation().RotateVector(EditedPosition.LeaveTangent); 
+		const FVector Tangent = SplineComp->GetComponentTransform().GetRotation().RotateVector(EditedPosition.ArriveTangent); // convert local-space tangent vectors to world-space
 
 		// Swap the tangents if they are not pointing in the same general direction
-		float CurrentAngle = FMath::Acos(FVector::DotProduct(ArriveTangent, NewArriveTangent) / (ArriveTangent.Size() * NewArriveTangent.Size()));
+		float CurrentAngle = FMath::Acos(FVector::DotProduct(Tangent, NewTangent) / (Tangent.Size() * NewTangent.Size()));
 		if (CurrentAngle > HALF_PI)
 		{
-			EditedPosition.ArriveTangent = SplineComp->GetComponentTransform().GetRotation().Inverse().RotateVector(NewLeaveTangent * -1.0f); // convert world-space tangent vectors back into local-space
-			EditedPosition.LeaveTangent = SplineComp->GetComponentTransform().GetRotation().Inverse().RotateVector(NewArriveTangent * -1.0f);
+			NewTangent = SplineComp->GetComponentTransform().GetRotation().Inverse().RotateVector(NewTangent * -1.0f); // convert world-space tangent vectors back into local-space
 		}
 		else
 		{
-			EditedPosition.ArriveTangent = SplineComp->GetComponentTransform().GetRotation().Inverse().RotateVector(NewArriveTangent); // convert world-space tangent vectors back into local-space
-			EditedPosition.LeaveTangent = SplineComp->GetComponentTransform().GetRotation().Inverse().RotateVector(NewLeaveTangent);
+			NewTangent = SplineComp->GetComponentTransform().GetRotation().Inverse().RotateVector(NewTangent); // convert world-space tangent vectors back into local-space
 		}
+
+		// Update tangent
+		EditedPosition.ArriveTangent = NewTangent;
+		EditedPosition.LeaveTangent = NewTangent;
+		EditedPosition.InterpMode = CIM_CurveUser;
   
-		// Copy rotation
-		FQuat NewRot = NearestSplineComp->GetComponentTransform().GetRotation() * NearestRotation.OutVal; // convert local-space rotation to world-space
+		// Copy rotation, it is only used to determine up vector so no need to adjust it 
+		FQuat NewRot = SplineComp->GetComponentTransform().GetRotation() * NearestRotation.OutVal; // convert local-space rotation to world-space
 		EditedRotation.OutVal = SplineComp->GetComponentTransform().GetRotation().Inverse() * NewRot; // convert world-space rotation to local-space
-	}
-  
-	// Copy scale - X is not used so ignore it
-	const FVector NearestSplineCompScale = NearestSplineComp->GetComponentTransform().GetScale3D();
-	const FVector SplineCompScale = SplineComp->GetComponentTransform().GetScale3D();
-	float NewScaleY = NearestSplineCompScale.Y * NearestScale.OutVal.Y; // convert local-space scale to world-space
-	float NewScaleZ = NearestSplineCompScale.Z * NearestScale.OutVal.Z;
-	EditedScale.OutVal.Y = FMath::IsNearlyZero(SplineCompScale.Y) ? NewScaleY : NewScaleY / SplineCompScale.Y; // convert world-space scale to local-space
-	EditedScale.OutVal.Z = FMath::IsNearlyZero(SplineCompScale.Z) ? NewScaleZ : NewScaleZ / SplineCompScale.Z; 
-
-	// Copy metadata
-	if (USplineMetadata* SplineMetadata = SplineComp->GetSplinePointsMetadata())
-	{ 
-		if (const USplineMetadata* NearestSplineMetadata = NearestSplineComp->GetSplinePointsMetadata())
+
+		// Copy scale, only when aligning parallel
+		if (InSnapMode == ESplineComponentSnapMode::AlignToTangent)
+		{
+			FVector NewScale = SplineComp->GetComponentTransform().GetScale3D() * NearestScale.OutVal; // convert local-space rotation to world-space
+			const FVector SplineCompScale = SplineComp->GetComponentTransform().GetScale3D();
+			EditedScale.OutVal.X = FMath::IsNearlyZero(SplineCompScale.X) ? NewScale.X : NewScale.X / SplineCompScale.X; // convert world-space scale to local-space
+			EditedScale.OutVal.Y = FMath::IsNearlyZero(SplineCompScale.Y) ? NewScale.Y : NewScale.Y / SplineCompScale.Y;
+			EditedScale.OutVal.Z = FMath::IsNearlyZero(SplineCompScale.Z) ? NewScale.Z : NewScale.Z / SplineCompScale.Z;
+		}
+
+		// Copy metadata
+		if (USplineMetadata* SplineMetadata = SplineComp->GetSplinePointsMetadata())
 		{ 
-			SplineMetadata->CopyPoint(NearestSplineMetadata, NearestKeyIndex, LastKeyIndexSelected);
+			if (const USplineMetadata* NearestSplineMetadata = NearestSplineComp->GetSplinePointsMetadata())
+			{ 
+				SplineMetadata->CopyPoint(NearestSplineMetadata, NearestKeyIndex, LastKeyIndexSelected);
+			}
 		}
 	}
 
@@ -1322,7 +1349,7 @@ void FSplineComponentVisualizer::OnSnapToNearestSplinePoint(bool bAlign)
 
 	NotifyPropertyModified(SplineComp, SplineCurvesProperty);
 				
-	if (bAlign)
+	if (InSnapMode == ESplineComponentSnapMode::AlignToTangent || InSnapMode == ESplineComponentSnapMode::AlignPerpendicularToTangent)
 	{
 		CachedRotation = SplineComp->GetQuaternionAtSplinePoint(LastKeyIndexSelected, ESplineCoordinateSpace::World);
 	}
@@ -2436,6 +2463,7 @@ void FSplineComponentVisualizer::GenerateSnapAlignSubMenu(FMenuBuilder& MenuBuil
 	MenuBuilder.AddMenuEntry(FLevelEditorCommands::Get().AlignToFloor);
 	MenuBuilder.AddMenuEntry(FSplineComponentVisualizerCommands::Get().SnapToNearestSplinePoint);
 	MenuBuilder.AddMenuEntry(FSplineComponentVisualizerCommands::Get().AlignToNearestSplinePoint);
+	MenuBuilder.AddMenuEntry(FSplineComponentVisualizerCommands::Get().AlignPerpendicularToNearestSplinePoint);
 	MenuBuilder.AddMenuEntry(FSplineComponentVisualizerCommands::Get().SnapAllToSelectedX);
 	MenuBuilder.AddMenuEntry(FSplineComponentVisualizerCommands::Get().SnapAllToSelectedY);
 	MenuBuilder.AddMenuEntry(FSplineComponentVisualizerCommands::Get().SnapAllToSelectedZ);
diff --git a/Engine/Source/Editor/ComponentVisualizers/Public/SplineComponentVisualizer.h b/Engine/Source/Editor/ComponentVisualizers/Public/SplineComponentVisualizer.h
--- a/Engine/Source/Editor/ComponentVisualizers/Public/SplineComponentVisualizer.h
+++ b/Engine/Source/Editor/ComponentVisualizers/Public/SplineComponentVisualizer.h
@@ -71,6 +71,17 @@ struct HSplineTangentHandleProxy : public HSplineVisProxy
 	bool bArriveTangent;
 };
 
+/** Accepted modes for snapping points. */
+namespace ESplineComponentSnapMode 
+{
+	enum Type
+	{
+		Snap,
+		AlignToTangent,
+		AlignPerpendicularToTangent
+	};
+}
+
 /** SplineComponent visualizer/edit functionality */
 class COMPONENTVISUALIZERS_API FSplineComponentVisualizer : public FComponentVisualizer
 {
@@ -151,7 +162,7 @@ protected:
 	void OnAddKeyToSegment();
 	bool CanAddKeyToSegment() const;
 
-	void OnSnapToNearestSplinePoint(bool bAlign);
+	void OnSnapToNearestSplinePoint(ESplineComponentSnapMode::Type InSnapMode);
 	bool CanSnapToNearestSplinePoint() const;
 
 	void OnSnapAll(EAxis::Type InAxis);
