Tools: Add panel to spline context menu to easily create simple shapes with splines, add options to set visibliity and scale of the cine rail preview mesh

diff --git a/Engine/Source/Editor/ComponentVisualizers/Private/SplineComponentVisualizer.cpp b/Engine/Source/Editor/ComponentVisualizers/Private/SplineComponentVisualizer.cpp
--- a/Engine/Source/Editor/ComponentVisualizers/Private/SplineComponentVisualizer.cpp
+++ b/Engine/Source/Editor/ComponentVisualizers/Private/SplineComponentVisualizer.cpp
@@ -17,6 +17,7 @@
 #include "ScopedTransaction.h"
 #include "ActorEditorUtils.h"
 #include "WorldCollision.h"
+#include "Widgets/Docking/SDockTab.h"
 
 IMPLEMENT_HIT_PROXY(HSplineVisProxy, HComponentVisProxy);
 IMPLEMENT_HIT_PROXY(HSplineKeyProxy, HSplineVisProxy);
@@ -136,6 +137,8 @@ public:
 	TSharedPtr<FUICommandInfo> ResetToDefault;
 };
 
+TWeakPtr<SWindow> FSplineComponentVisualizer::WeakExistingWindow;
+
 FSplineComponentVisualizer::FSplineComponentVisualizer()
 	: FComponentVisualizer()
 	, LastKeyIndexSelected(INDEX_NONE)
@@ -528,6 +531,11 @@ void FSplineComponentVisualizer::ChangeSelectionState(int32 Index, bool bIsCtrlH
 			LastKeyIndexSelected = Index;
 		}
 	}
+
+	if (SplineGeneratorPanel.IsValid())
+	{
+		SplineGeneratorPanel->OnSelectionUpdated();
+	}
 }
 
 bool FSplineComponentVisualizer::VisProxyHandleClick(FEditorViewportClient* InViewportClient, HComponentVisProxy* VisProxy, const FViewportClick& Click)
@@ -660,7 +668,6 @@ USplineComponent* FSplineComponentVisualizer::GetEditedSplineComponent() const
 	return Cast<USplineComponent>(SplinePropertyPath.GetComponent());
 }
 
-
 bool FSplineComponentVisualizer::GetWidgetLocation(const FEditorViewportClient* ViewportClient, FVector& OutLocation) const
 {
 	USplineComponent* SplineComp = GetEditedSplineComponent();
@@ -691,7 +698,7 @@ bool FSplineComponentVisualizer::GetWidgetLocation(const FEditorViewportClient*
 			// Otherwise use the last key index set
 			check(LastKeyIndexSelected < Position.Points.Num());
 			check(SelectedKeys.Contains(LastKeyIndexSelected));
-			const auto& Point = Position.Points[LastKeyIndexSelected];
+			const FInterpCurvePointVector& Point = Position.Points[LastKeyIndexSelected];
 			OutLocation = SplineComp->GetComponentTransform().TransformPosition(Point.OutVal);
 			if (!DuplicateDelayAccumulatedDrag.IsZero())
 			{
@@ -2404,6 +2411,16 @@ TSharedPtr<SWidget> FSplineComponentVisualizer::GenerateContextMenu() const
 					}
 				}
 			}
+
+			MenuBuilder.AddMenuEntry(
+				LOCTEXT("SplineGenerate", "Spline Generation Panel"),
+				LOCTEXT("SplineGenerateTooltip", "Opens up a spline generation panel to easily create basic shapes with splines"),
+				FSlateIcon(),
+				FUIAction( 
+					FExecuteAction::CreateRaw(this, &FSplineComponentVisualizer::CreateSplineGeneratorPanel), 
+					FCanExecuteAction::CreateLambda([] { return true; })
+				)
+			);
 		}
 		MenuBuilder.EndSection();
 
@@ -2476,6 +2493,41 @@ void FSplineComponentVisualizer::GenerateLockAxisSubMenu(FMenuBuilder& MenuBuild
 	MenuBuilder.AddMenuEntry(FSplineComponentVisualizerCommands::Get().SetLockedAxisZ);
 }
 
+void FSplineComponentVisualizer::CreateSplineGeneratorPanel()
+{
+	SAssignNew(SplineGeneratorPanel, SSplineGeneratorPanel, SharedThis(this));
+
+	TSharedPtr<SWindow> ExistingWindow = WeakExistingWindow.Pin();
+	if (!ExistingWindow.IsValid())
+	{
+		ExistingWindow = SNew(SWindow)
+			.ScreenPosition(FSlateApplication::Get().GetCursorPos())
+			.Title(FText::FromString("Spline Generation"))
+			.SizingRule(ESizingRule::Autosized)
+			.AutoCenter(EAutoCenter::None);
+
+		ExistingWindow->SetOnWindowClosed(FOnWindowClosed::CreateSP(SplineGeneratorPanel.ToSharedRef(), &SSplineGeneratorPanel::OnWindowClosed));
+
+		TSharedPtr<SWindow> RootWindow = FSlateApplication::Get().GetActiveTopLevelWindow();
+
+		if (RootWindow.IsValid())
+		{
+			FSlateApplication::Get().AddWindowAsNativeChild(ExistingWindow.ToSharedRef(), RootWindow.ToSharedRef());
+		}
+		else
+		{
+			FSlateApplication::Get().AddWindow(ExistingWindow.ToSharedRef());
+		}
+
+		ExistingWindow->BringToFront();
+		WeakExistingWindow = ExistingWindow;
+	}
+	else
+	{
+		ExistingWindow->BringToFront();
+	}
+	ExistingWindow->SetContent(SplineGeneratorPanel.ToSharedRef());
+}
 
 
 #undef LOCTEXT_NAMESPACE
diff --git a/Engine/Source/Editor/ComponentVisualizers/Private/SplineGeneratorPanel.cpp b/Engine/Source/Editor/ComponentVisualizers/Private/SplineGeneratorPanel.cpp
new file mode 100644
--- /dev/null
+++ b/Engine/Source/Editor/ComponentVisualizers/Private/SplineGeneratorPanel.cpp
@@ -0,0 +1,618 @@
+// Copyright 1998-2019 Epic Games, Inc. All Rights Reserved.
+
+#include "SplineGeneratorPanel.h"
+#include "Widgets/Input/SButton.h"
+#include "Layout/WidgetPath.h"
+#include "Widgets/Input/SCheckBox.h"
+#include "Widgets/Text/STextBlock.h"
+#include "ScopedTransaction.h"
+
+#define LOCTEXT_NAMESPACE "SplineGenerator"
+
+void USplineGeneratorBase::Init(TWeakPtr<FSplineComponentVisualizer> InWeakSplineVis)
+{
+	check(InWeakSplineVis.IsValid());
+	WeakSplineVis = InWeakSplineVis;
+
+	TSharedPtr<FSplineComponentVisualizer> SplineVis = WeakSplineVis.Pin();
+
+	SelectedSplineComponent = MakeWeakObjectPtr(SplineVis->GetEditedSplineComponent());
+
+	StartKey = INDEX_NONE;
+	check(SplineVis->GetSelectedKeys().Num() == 1);
+	StartKey = *SplineVis->GetSelectedKeys().CreateConstIterator();
+	CachedSplineCurves = SelectedSplineComponent->SplineCurves;
+}
+
+void USplineGeneratorBase::Reset()
+{
+	if (!SelectedSplineComponent.IsValid())
+	{
+		TSharedPtr<FSplineComponentVisualizer> SplineVis = WeakSplineVis.Pin();
+		check(SplineVis);
+
+		SelectedSplineComponent = MakeWeakObjectPtr(SplineVis->GetEditedSplineComponent());
+	}
+
+	if (SelectedSplineComponent.IsValid())
+	{
+		SelectedSplineComponent->Modify();
+		if (AActor* Owner = SelectedSplineComponent->GetOwner())
+		{
+			Owner->Modify();
+		}
+
+		SelectedSplineComponent->SplineCurves = CachedSplineCurves;
+
+		CachedSplineCurves = FSplineCurves();
+
+		SelectedSplineComponent->UpdateSpline();
+		SelectedSplineComponent->bSplineHasBeenEdited = true;
+
+		UProperty* SplineCurvesProperty = FindField<UProperty>(USplineComponent::StaticClass(), GET_MEMBER_NAME_CHECKED(USplineComponent, SplineCurves));
+		FComponentVisualizer::NotifyPropertyModified(SelectedSplineComponent.Get(), SplineCurvesProperty);
+
+		GEditor->RedrawLevelEditingViewports(true);
+	}
+}
+
+void USplineGeneratorBase::PreviewCurve()
+{
+	if (!SelectedSplineComponent.IsValid())
+	{
+		TSharedPtr<FSplineComponentVisualizer> SplineVis = WeakSplineVis.Pin();
+		check(SplineVis);
+
+		SelectedSplineComponent = MakeWeakObjectPtr(SplineVis->GetEditedSplineComponent());
+	}
+
+	if (SelectedSplineComponent.IsValid())
+	{
+		SelectedSplineComponent->Modify();
+		if (AActor* Owner = SelectedSplineComponent->GetOwner())
+		{
+			Owner->Modify();
+		}
+
+		BuildCurve();
+
+		SelectedSplineComponent->UpdateSpline();
+		SelectedSplineComponent->bSplineHasBeenEdited = true;
+
+		UProperty* SplineCurvesProperty = FindField<UProperty>(USplineComponent::StaticClass(), GET_MEMBER_NAME_CHECKED(USplineComponent, SplineCurves));
+		FComponentVisualizer::NotifyPropertyModified(SelectedSplineComponent.Get(), SplineCurvesProperty);
+
+		GEditor->RedrawLevelEditingViewports(true);
+	}
+}
+
+int32 USplineGeneratorBase::GetAddIndex(int32 Index) const
+{
+	switch (ShapeAddMode)
+	{
+	case EShapeAddMode::AppendAfter:
+		return SelectedSplineComponent->GetNumberOfSplinePoints();
+	case EShapeAddMode::AppendBefore:
+		return 0;
+	case EShapeAddMode::InsertAfter:
+		return StartKey + Index + 1;
+	case EShapeAddMode::InsertBefore:
+		return StartKey;
+	}
+
+	return 0;
+}
+
+int32 USplineGeneratorBase::GetItrIndex(int32 Index) const
+{
+	switch (ShapeAddMode)
+	{
+	case EShapeAddMode::AppendAfter:
+		return SelectedSplineComponent->GetNumberOfSplinePoints() - GetNumPoints() + Index;
+	case EShapeAddMode::AppendBefore:
+		return GetNumPoints() - Index - 1;
+	case EShapeAddMode::InsertAfter:
+		return StartKey + Index + 1;
+	case EShapeAddMode::InsertBefore:
+		return StartKey + GetNumPoints() - Index - 1;
+	}
+
+	return 0;
+}
+
+/*
+ * Fancy math to recalculate tangents to turn into circle
+ */
+double CalcTangentMultiplier(const float InRadius, const float InRotInc)
+{
+	static constexpr double A = .5f;
+	static constexpr double A2 = A * A;
+	static constexpr double A3 = A2 * A;
+
+	// Use first and second keys added as a sample calculation
+	const FVector T0 = FVector::ForwardVector;
+	const FVector T1 = T0.RotateAngleAxis(InRotInc, FVector::UpVector);
+	const FVector P0 = FVector::RightVector * InRadius;
+	const FVector P1 = P0.RotateAngleAxis(InRotInc, FVector::UpVector);
+
+	// Calculate the difference between the actual interpolated midpoint and expected interpolated midpoint
+	const FVector ActualVal = FMath::CubicInterp(P0, T0, P1, T1, A);
+	const FVector ExpectedVal = P0.RotateAngleAxis(InRotInc * A, FVector::UpVector);
+	const float Diff = (ActualVal.X - ExpectedVal.X);
+
+	// Do a partial calculation of the cubic interpolation equation
+	static constexpr double C1 = (A3 - (2 * A2) + A), C2 = (A3 - A2);
+	const double PartialInterp = -1.f * ((C1 * T0.X) + (C2 * T1.X));
+
+	// Calculate the final multiplier to multiply to all normalized tangents
+	return FMath::IsNearlyZero(PartialInterp) ? 1.f : ((Diff / PartialInterp) + 1.f);
+}
+
+void UCircleSplineGenerator::BuildCurve()
+{
+	// Re-set the cached spline keys so we can add to them
+	SelectedSplineComponent->SplineCurves = CachedSplineCurves;
+
+	// naming is a little weird, will fix
+	const float BranchRightFlip = ((bKeepFirstKeyTangent && bBranchRight) ? 1.f : -1.f);
+	const float ReverseDirFlip = ((bReverseDir) ^ !(ShapeAddMode & (EShapeAddMode::AppendBefore | EShapeAddMode::InsertBefore)) ? -1.f : 1.f);
+
+	// Find starting features
+	const FVector StartPoint = SelectedSplineComponent->GetLocationAtSplinePoint(GetAddIndex(-1), ESplineCoordinateSpace::Local);
+	const FQuat StartQuat = SelectedSplineComponent->GetQuaternionAtSplinePoint(GetAddIndex(-1), ESplineCoordinateSpace::Local);
+
+	const FVector CenterDir = bKeepFirstKeyTangent ? (StartQuat.GetRightVector() * BranchRightFlip) : StartQuat.GetForwardVector();
+	const FVector CircleCenter = StartPoint + Radius * CenterDir;
+	const float RotIncrement = (360.f / float(NumberOfPoints)) * ReverseDirFlip * BranchRightFlip * -1.f;
+
+	// Add points around circle
+	FVector CircleItr = StartPoint;
+	for (int32 Index = 0; Index < NumberOfPoints; Index++)
+	{
+		const FVector Diff = CircleItr - CircleCenter;
+		CircleItr = CircleCenter + Diff.RotateAngleAxis(RotIncrement, StartQuat.GetUpVector());
+
+		SelectedSplineComponent->AddSplinePointAtIndex(CircleItr, GetAddIndex(Index), ESplineCoordinateSpace::Local);
+	}
+
+	FInterpCurveQuat& SplineRotation = SelectedSplineComponent->GetSplinePointsRotation();
+
+	float StartTangentFlip = BranchRightFlip * (bReverseDir ? 1.f : -1.f);
+	const FVector StartCircleTangent = (bKeepFirstKeyTangent ? StartQuat.GetForwardVector() : StartQuat.GetRightVector()) * StartTangentFlip;
+
+	if (NumberOfPoints >= 2)
+	{
+		const double TangentMult = CalcTangentMultiplier(Radius, RotIncrement) * ReverseDirFlip * -1.f;
+		for (int32 Index = NumberOfPoints - 1; Index >= -1; Index--)
+		{
+			const FVector Tangent = StartCircleTangent.RotateAngleAxis(RotIncrement * float(Index + 1), StartQuat.GetUpVector()).GetSafeNormal();
+			SelectedSplineComponent->SetTangentAtSplinePoint(GetItrIndex(Index), Tangent * TangentMult, ESplineCoordinateSpace::Local);
+			SplineRotation.Points[GetItrIndex(Index)].OutVal = (FQuat::MakeFromEuler(FVector(0.f, 0.f, RotIncrement * float(Index + 1))) * StartQuat);
+		}
+	}
+}
+
+void UArcSplineGenerator::BuildCurve()
+{
+	// Re-set the cached spline keys so we can add to them
+	SelectedSplineComponent->SplineCurves = CachedSplineCurves;
+
+	// Prevent the curve from going whack
+	Degrees = FMath::Clamp(Degrees, 0.f, 180.f * float(NumberOfPoints));
+
+	const float BranchRightFlip = ((bKeepFirstKeyTangent && bBranchRight) ? 1.f : -1.f);
+	const float ReverseDirFlip = ((bReverseDir) ^ !(ShapeAddMode & (EShapeAddMode::AppendBefore | EShapeAddMode::InsertBefore)) ? -1.f : 1.f) * BranchRightFlip * -1.f;
+	// Find starting features
+	const FVector StartPoint = SelectedSplineComponent->GetLocationAtSplinePoint(GetAddIndex(-1), ESplineCoordinateSpace::Local);
+	const FQuat StartQuat = SelectedSplineComponent->GetQuaternionAtSplinePoint(GetAddIndex(-1), ESplineCoordinateSpace::Local);
+
+	const FVector CenterDir = bKeepFirstKeyTangent ? StartQuat.GetRightVector() * BranchRightFlip : StartQuat.GetForwardVector();
+	const FVector CircleCenter = StartPoint + Radius * CenterDir;
+	const float RotIncrement = (Degrees / float(NumberOfPoints)) * ReverseDirFlip;
+
+	// Add points around arc
+	FVector CircleItr = StartPoint;
+	for (int32 Index = 0; Index < NumberOfPoints; Index++)
+	{
+		const FVector Diff = CircleItr - CircleCenter;
+		CircleItr = CircleCenter + Diff.RotateAngleAxis(RotIncrement, StartQuat.GetUpVector());
+
+		SelectedSplineComponent->AddSplinePointAtIndex(CircleItr, GetAddIndex(Index), ESplineCoordinateSpace::Local);
+	}
+
+	FInterpCurveQuat& SplineRotation = SelectedSplineComponent->GetSplinePointsRotation();
+
+	float StartTangentFlip = (bReverseDir ? -1.f : 1.f);
+ 	const FVector StartCircleTangent = (bKeepFirstKeyTangent ? StartQuat.GetForwardVector() : StartQuat.GetRightVector()) * StartTangentFlip;
+
+	if (NumberOfPoints >= 2)
+	{
+		const double TangentMult = CalcTangentMultiplier(Radius, RotIncrement) * ReverseDirFlip * -1.f;
+
+		SelectedSplineComponent->SetTangentAtSplinePoint(GetItrIndex(-1), StartCircleTangent * TangentMult, ESplineCoordinateSpace::Local);
+		for (int32 Index = 0; Index < NumberOfPoints; Index++)
+		{
+			const FVector Tangent = StartCircleTangent.RotateAngleAxis(RotIncrement * float(Index + 1), StartQuat.GetUpVector()).GetSafeNormal();
+			SelectedSplineComponent->SetTangentAtSplinePoint(GetItrIndex(Index), Tangent * TangentMult, ESplineCoordinateSpace::Local);
+			SplineRotation.Points[GetItrIndex(Index)].OutVal = (FQuat::MakeFromEuler(FVector(0.f, 0.f, RotIncrement * float(Index + 1))) * StartQuat);
+		}
+	}
+}
+
+void UEllipseSplineGenerator::BuildCurve()
+{
+	// Re-set the cached spline keys so we can add to them
+	SelectedSplineComponent->SplineCurves = CachedSplineCurves;
+
+	const float BranchRightFlip = ((bKeepFirstKeyTangent && bBranchRight) ? 1.f : -1.f);
+	const float ReverseDirFlip = (bReverseDir ^ !(ShapeAddMode & (EShapeAddMode::AppendBefore | EShapeAddMode::InsertBefore)) ? -1.f : 1.f) * BranchRightFlip * -1.f;
+
+	// Find starting features
+	const FVector StartPoint = SelectedSplineComponent->GetLocationAtSplinePoint(GetAddIndex(-1), ESplineCoordinateSpace::Local);
+	const FQuat StartQuat = SelectedSplineComponent->GetQuaternionAtSplinePoint(GetAddIndex(-1), ESplineCoordinateSpace::Local);
+
+	const FVector CenterDir = bKeepFirstKeyTangent ? StartQuat.GetRightVector() * BranchRightFlip : StartQuat.GetForwardVector();
+	const FVector CircleCenter = StartPoint + (bKeepFirstKeyTangent ? Length * .5f : Width * .5f) * CenterDir;
+
+	FInterpCurveQuat& SplineRotation = SelectedSplineComponent->GetSplinePointsRotation();
+
+	// Add points around ellipse
+	float CycleInc = (2 * PI / float(NumberOfPoints)) * ReverseDirFlip;
+	float CycleItr = CycleInc + (bKeepFirstKeyTangent ? PI * .5f : PI) * BranchRightFlip * -1.f;
+	const float TangentFlip = ReverseDirFlip * (!(ShapeAddMode & (EShapeAddMode::AppendBefore | EShapeAddMode::InsertBefore)) ? -1.f : 1.f) * -1.f;
+
+	const FVector StartCircleTangent = (bKeepFirstKeyTangent ? StartQuat.GetForwardVector() : StartQuat.GetRightVector()) * (bReverseDir ? -1.f : 1.f);
+	const FVector PrevStartTangent = SelectedSplineComponent->GetTangentAtSplinePoint(GetItrIndex(-1), ESplineCoordinateSpace::Local);
+	SelectedSplineComponent->SetTangentAtSplinePoint(GetItrIndex(-1), StartCircleTangent * PrevStartTangent.Size(), ESplineCoordinateSpace::Local);
+	for (int32 Index = 0; Index < NumberOfPoints; Index++, CycleItr += CycleInc)
+	{
+		const FVector EllipsePoint = CircleCenter + FVector(Width * .5f * FMath::Cos(CycleItr), Length * .5f * FMath::Sin(CycleItr), 0.f);
+		const FVector EllipseTangent = FVector(Width * .5f * -1.f * FMath::Sin(CycleItr), Length * .5f * FMath::Cos(CycleItr), 0.f).GetSafeNormal();
+		const FVector RotatedEllipsePoint = CircleCenter + StartQuat.Rotator().RotateVector(EllipsePoint - CircleCenter);
+		const FVector RotatedEllipseTangent = StartQuat.Rotator().RotateVector(EllipseTangent) * TangentFlip;
+		const int32 AddIdx = GetAddIndex(Index);
+		SelectedSplineComponent->AddSplinePointAtIndex(RotatedEllipsePoint, AddIdx, ESplineCoordinateSpace::Local);
+		const FVector PrevTangent = SelectedSplineComponent->GetTangentAtSplinePoint(AddIdx, ESplineCoordinateSpace::Local);
+		SelectedSplineComponent->SetTangentAtSplinePoint(AddIdx, RotatedEllipseTangent * PrevTangent.Size(), ESplineCoordinateSpace::Local);
+		SplineRotation.Points[AddIdx].OutVal = (FQuat::MakeFromEuler(FVector(0.f, 0.f, CycleItr)) * StartQuat);
+	}
+}
+
+void USquareSplineGenerator::BuildCurve()
+{
+	// Re-set the cached spline keys so we can add to them
+	SelectedSplineComponent->SplineCurves = CachedSplineCurves;
+
+	const bool bPrepend = !(ShapeAddMode & (EShapeAddMode::AppendAfter | EShapeAddMode::InsertAfter));
+	const float LengthFlip = bPrepend ? -1.f : 1.f;
+
+	// Find starting features
+	const FVector StartPoint = SelectedSplineComponent->GetLocationAtSplinePoint(GetAddIndex(-1), ESplineCoordinateSpace::Local);
+	const FQuat StartQuat = SelectedSplineComponent->GetQuaternionAtSplinePoint(GetAddIndex(-1), ESplineCoordinateSpace::Local);
+
+	FInterpCurveQuat& SplineRotation = SelectedSplineComponent->GetSplinePointsRotation();
+
+	const float RotDir = bBranchRight ? 90.f : -90.f;
+
+	FVector RectItr = StartPoint;
+	FVector DirItr = StartQuat.GetForwardVector();
+
+	for (int32 Index = 0; Index < 4; Index++)
+	{
+		RectItr += DirItr * Length * LengthFlip;
+		const int32 AddIdx = GetAddIndex(Index);
+		SelectedSplineComponent->AddSplinePointAtIndex(RectItr, AddIdx, ESplineCoordinateSpace::Local);
+		SelectedSplineComponent->SplineCurves.Position.Points[bPrepend ? AddIdx : AddIdx - 1].InterpMode = EInterpCurveMode::CIM_Linear;
+		SplineRotation.Points[AddIdx].OutVal = FQuat::MakeFromEuler(FVector(0.f, 0.f, float(Index) * RotDir)) * StartQuat;
+		DirItr = DirItr.RotateAngleAxis(RotDir, StartQuat.GetUpVector());
+	}
+}
+
+void URectangleSplineGenerator::BuildCurve()
+{
+	// Re-set the cached spline keys so we can add to them
+	SelectedSplineComponent->SplineCurves = CachedSplineCurves;
+
+	const bool bPrepend = !(ShapeAddMode & (EShapeAddMode::AppendAfter | EShapeAddMode::InsertAfter));
+	const float LengthFlip = bPrepend ? -1.f : 1.f;
+
+	// Find starting features
+	const FVector StartPoint = SelectedSplineComponent->GetLocationAtSplinePoint(GetAddIndex(-1), ESplineCoordinateSpace::Local);
+	const FQuat StartQuat = SelectedSplineComponent->GetQuaternionAtSplinePoint(GetAddIndex(-1), ESplineCoordinateSpace::Local);
+
+	FInterpCurveQuat& SplineRotation = SelectedSplineComponent->GetSplinePointsRotation();
+
+	const float RotDir = bBranchRight ? 90.f : -90.f;
+
+	int32 Index = 0;
+	FVector RectItr = StartPoint;
+	FVector DirItr = StartQuat.GetForwardVector();
+
+	{
+		// Draw first segment
+		const int32 AddIdx = GetAddIndex(Index);
+		RectItr += DirItr * Length * LengthFlip;
+		SelectedSplineComponent->AddSplinePointAtIndex(RectItr, AddIdx, ESplineCoordinateSpace::Local);
+		SelectedSplineComponent->SplineCurves.Position.Points[bPrepend ? AddIdx : AddIdx - 1].InterpMode = EInterpCurveMode::CIM_Linear;
+		SplineRotation.Points[AddIdx].OutVal = FQuat::MakeFromEuler(FVector(0.f, 0.f, float(Index) * RotDir)) * StartQuat;
+		Index++;
+	}
+
+	{
+		// Draw second segment
+		const int32 AddIdx = GetAddIndex(Index);
+		DirItr = DirItr.RotateAngleAxis(RotDir, StartQuat.GetUpVector());
+		RectItr += DirItr * Width * LengthFlip;
+		SelectedSplineComponent->AddSplinePointAtIndex(RectItr, AddIdx, ESplineCoordinateSpace::Local);
+		SelectedSplineComponent->SplineCurves.Position.Points[bPrepend ? AddIdx : AddIdx - 1].InterpMode = EInterpCurveMode::CIM_Linear;
+		SplineRotation.Points[AddIdx].OutVal = FQuat::MakeFromEuler(FVector(0.f, 0.f, float(Index) * RotDir)) * StartQuat;
+		Index++;
+	}
+
+	{
+		// Draw third segment
+		const int32 AddIdx = GetAddIndex(Index);
+		DirItr = DirItr.RotateAngleAxis(RotDir, StartQuat.GetUpVector());
+		RectItr += DirItr * Length * LengthFlip;
+		SelectedSplineComponent->AddSplinePointAtIndex(RectItr, AddIdx, ESplineCoordinateSpace::Local);
+		SelectedSplineComponent->SplineCurves.Position.Points[bPrepend ? AddIdx : AddIdx - 1].InterpMode = EInterpCurveMode::CIM_Linear;
+		SplineRotation.Points[AddIdx].OutVal = FQuat::MakeFromEuler(FVector(0.f, 0.f, float(Index) * RotDir)) * StartQuat;
+		Index++;
+	}
+
+	{
+		// Draw fourth segment
+		const int32 AddIdx = GetAddIndex(Index);
+		DirItr = DirItr.RotateAngleAxis(RotDir, StartQuat.GetUpVector());
+		RectItr += DirItr * Width * LengthFlip;
+		SelectedSplineComponent->AddSplinePointAtIndex(RectItr, AddIdx, ESplineCoordinateSpace::Local);
+		SelectedSplineComponent->SplineCurves.Position.Points[bPrepend ? AddIdx : AddIdx - 1].InterpMode = EInterpCurveMode::CIM_Linear;
+		SplineRotation.Points[AddIdx].OutVal = FQuat::MakeFromEuler(FVector(0.f, 0.f, float(Index) * RotDir)) * StartQuat;
+		Index++;
+	}
+}
+
+void ULineSplineGenerator::BuildCurve()
+{
+	// Re-set the cached spline keys so we can add to them
+	SelectedSplineComponent->SplineCurves = CachedSplineCurves;
+
+	const bool bPrepend = !(ShapeAddMode & (EShapeAddMode::AppendAfter | EShapeAddMode::InsertAfter));
+	const float LengthFlip = bPrepend ? -1.f : 1.f;
+
+	// Find starting features
+	const FVector StartPoint = SelectedSplineComponent->GetLocationAtSplinePoint(GetAddIndex(-1), ESplineCoordinateSpace::Local);
+	const FQuat StartQuat = SelectedSplineComponent->GetQuaternionAtSplinePoint(GetAddIndex(-1), ESplineCoordinateSpace::Local);
+
+	float LineInc = 0.0f;
+	FVector LineDir = FVector::ZeroVector;
+	
+	bEnableUpToNextPoint = false;
+	int32 NextIdx = 0;
+	if (ShapeAddMode == EShapeAddMode::InsertAfter || ShapeAddMode == EShapeAddMode::InsertBefore)
+	{
+		NextIdx = !(ShapeAddMode & EShapeAddMode::InsertAfter) ? StartKey - 1 : StartKey + 1;
+		if (NextIdx >= 0 && NextIdx < SelectedSplineComponent->GetNumberOfSplinePoints())
+		{
+			bEnableUpToNextPoint = true;
+		}
+	}
+
+	if (bUpToNextPoint && bEnableUpToNextPoint)
+	{
+		const FVector NextPoint = SelectedSplineComponent->GetLocationAtSplinePoint(NextIdx, ESplineCoordinateSpace::Local);
+		const FVector Diff = NextPoint - StartPoint;
+		Length = Diff.Size();
+		LineDir = Diff.GetSafeNormal();
+		LineInc = Diff.Size() / float(NumberOfPoints + 1);
+	}
+	else
+	{
+		LineInc = Length / float(NumberOfPoints);
+		LineDir = StartQuat.GetForwardVector() * LengthFlip;
+	}
+
+	for (int32 Index = 0; Index < NumberOfPoints; Index++)
+	{
+		const int32 AddIdx = GetAddIndex(Index);
+		const FVector NewPoint = StartPoint + (LineInc * float(Index + 1)) * LineDir;
+		SelectedSplineComponent->AddSplinePointAtIndex(NewPoint, AddIdx, ESplineCoordinateSpace::Local);
+	}
+}
+
+void AddGeneratorsToRegistry(TArray<UObject*>& InShapeGenRegistry)
+{
+	InShapeGenRegistry.Add(NewObject<UCircleSplineGenerator>());
+	InShapeGenRegistry.Add(NewObject<UArcSplineGenerator>());
+	InShapeGenRegistry.Add(NewObject<USquareSplineGenerator>());
+	InShapeGenRegistry.Add(NewObject<UEllipseSplineGenerator>());
+	InShapeGenRegistry.Add(NewObject<URectangleSplineGenerator>());
+	InShapeGenRegistry.Add(NewObject<ULineSplineGenerator>());
+
+	for (auto ShapeGen : InShapeGenRegistry)
+	{
+		ShapeGen->ClearFlags(RF_Transactional);
+		ShapeGen->AddToRoot();
+	}
+}
+
+void SSplineGeneratorPanel::Construct(const FArguments& InArgs, TWeakPtr<FSplineComponentVisualizer> InWeakSplineComponentVisualizer)
+{
+	WeakSplineComponentVisualizer = InWeakSplineComponentVisualizer;
+	TSharedPtr<FSplineComponentVisualizer> SplineVis = WeakSplineComponentVisualizer.Pin();
+	check(SplineVis);
+
+	AddGeneratorsToRegistry(ShapeGenRegistry);
+
+	FPropertyEditorModule& PropertyEditor = FModuleManager::LoadModuleChecked<FPropertyEditorModule>("PropertyEditor");
+	FClassViewerModule& ClassViewerModule = FModuleManager::LoadModuleChecked<FClassViewerModule>("ClassViewer");
+
+	// Configure the Details View
+	FDetailsViewArgs DetailsViewArgs;
+	DetailsViewArgs.bUpdatesFromSelection = false;
+	DetailsViewArgs.bLockable = false;
+	DetailsViewArgs.ViewIdentifier = "SplineGenerationOptions";
+	DetailsViewArgs.DefaultsOnlyVisibility = EEditDefaultsOnlyNodeVisibility::Automatic;
+	DetailsViewArgs.bAllowSearch = false;
+	DetailsViewArgs.bShowPropertyMatrixButton = false;
+	DetailsViewArgs.bShowOptions = true;
+
+	// Generate it and store a reference so we can update it with the right object later.
+	DetailView = PropertyEditor.CreateDetailView(DetailsViewArgs);
+
+	ActiveTransaction = MakeUnique<FScopedTransaction>(TEXT("SplineGenerator"), LOCTEXT("SplineGeneratorTransaction", "Generate Spline"), nullptr);
+
+	TSharedRef<SVerticalBox> ShapeSelectorWidget = SNew(SVerticalBox);
+	for (auto GenItr = ShapeGenRegistry.CreateConstIterator(); GenItr; ++GenItr)
+	{
+		auto OnRadioChanged = [this, GenObject = *GenItr] (ECheckBoxState CheckBoxState)
+		{
+			if (SplineGen)
+			{
+				SplineGen->Reset();
+				SplineGen = nullptr;
+				DetailView->SetObject(nullptr);
+			}
+			else
+			{
+				ActiveTransaction = MakeUnique<FScopedTransaction>(TEXT("SplineGenerator"), LOCTEXT("SplineGeneratorTransaction", "Generate Spline"), nullptr);
+			}
+
+			TSharedPtr<FSplineComponentVisualizer> SplineVis = WeakSplineComponentVisualizer.Pin();
+			if (SplineVis.IsValid() && SplineVis->GetSelectedKeys().Num() == 1)
+			{
+				
+				SplineGen = Cast<USplineGeneratorBase>(GenObject);
+				check(SplineGen);
+				SplineGen->Init(SplineVis);
+				DetailView->SetObject(SplineGen);
+			}
+		};
+
+		auto IsRadioChecked = [this, GenObject = *GenItr]()
+		{
+			if (SplineGen)
+			{
+				return GenObject == SplineGen ? ECheckBoxState::Checked : ECheckBoxState::Unchecked;
+			}
+			return ECheckBoxState::Undetermined;
+		};
+
+		ShapeSelectorWidget->AddSlot()
+		[
+			SNew(SCheckBox)
+			.Style(FEditorStyle::Get(), "RadioButton")
+			.IsChecked_Lambda(IsRadioChecked)
+			.OnCheckStateChanged_Lambda(OnRadioChanged)
+			.Padding(FMargin(10.f, 3.f))
+			.Content()
+			[
+				SNew(STextBlock)
+				.Text((*GenItr)->GetClass()->GetDisplayNameText())
+				.Font(FCoreStyle::GetDefaultFontStyle("Regular", 10))
+			]
+		];
+	}
+
+	ChildSlot
+	[
+		SNew(SVerticalBox)
+		+ SVerticalBox::Slot()
+		.HAlign(HAlign_Center)
+		.AutoHeight()
+		[
+			SNew(STextBlock)
+			.Text(LOCTEXT("Shapes", "Shapes"))
+			.Font(FCoreStyle::GetDefaultFontStyle("Regular", 11))
+		]
+		+ SVerticalBox::Slot()
+		.AutoHeight()
+		[
+			SNew(SBorder)
+			.Padding(FMargin(5.f, 5.f))
+			.BorderImage(FEditorStyle::GetBrush("ToolPanel.DarkGroupBorder"))
+			.Content()
+			[
+				ShapeSelectorWidget
+			]
+		]
+
+		+ SVerticalBox::Slot()
+		.FillHeight(1.f)
+		[
+			DetailView.ToSharedRef()
+		]
+	];
+}
+
+void SSplineGeneratorPanel::OnWindowClosed(const TSharedRef<SWindow>&)
+{
+	USplineComponent* SplineComp = WeakSplineComponentVisualizer.Pin()->GetEditedSplineComponent();
+
+	if (SplineGen)
+	{
+		SplineGen = nullptr;
+		DetailView->SetObject(nullptr);
+	}
+
+	ActiveTransaction.Reset();
+}
+
+void SSplineGeneratorPanel::OnSelectionUpdated()
+{
+	TSharedPtr<FSplineComponentVisualizer> SplineVis = WeakSplineComponentVisualizer.Pin();
+	if (SplineGen && SplineVis.IsValid())
+	{
+
+		if (SplineVis->GetSelectedKeys().Num() > 0 && *SplineVis->GetSelectedKeys().CreateConstIterator() != SplineGen->StartKey)
+		{
+			ActiveTransaction.Reset();
+
+			if (SplineVis->GetSelectedKeys().Num() == 1)
+			{	
+				ActiveTransaction  = MakeUnique<FScopedTransaction>(TEXT("SplineGenerator"), LOCTEXT("SplineGeneratorTransaction", "Generate Spline"), nullptr);
+				SplineGen->Init(SplineVis);
+			}
+			else
+			{
+				SplineGen = nullptr;
+				DetailView->SetObject(nullptr);
+			}
+		}
+		else if (SplineVis->GetSelectedKeys().Num() == 0)
+		{
+			SplineGen = nullptr;
+			DetailView->SetObject(nullptr);
+			ActiveTransaction.Reset();
+		}
+	}
+}
+
+void SSplineGeneratorPanel::Tick(const FGeometry& AllottedGeometry, const double InCurrentTime, const float InDeltaTime)
+{
+	if (SplineGen)
+	{
+		SplineGen->PreviewCurve();
+	}
+}
+
+void SSplineGeneratorPanel::OnFocusChanging(const FWeakWidgetPath& PreviousFocusPath, const FWidgetPath& NewWidgetPath, const FFocusEvent& InFocusEvent)
+{
+	FWidgetPath WidgetPath = NewWidgetPath.GetPathDownTo(SharedThis(this));
+	FWidgetPath WidgetPath2 = NewWidgetPath;
+	bool bPathFound = WidgetPath2.ExtendPathTo(FWidgetMatcher(SharedThis(this)));
+	if (SplineGen && !WidgetPath.IsValid() && bPathFound)
+	{
+		SplineGen = nullptr;
+		DetailView->SetObject(nullptr);
+		ActiveTransaction.Reset();
+	}
+}
+
+SSplineGeneratorPanel::~SSplineGeneratorPanel()
+{
+	if (ActiveTransaction.IsValid() && ActiveTransaction->IsOutstanding())
+	{
+		ActiveTransaction ->Cancel();
+	}
+}
\ No newline at end of file
diff --git a/Engine/Source/Editor/ComponentVisualizers/Public/SplineComponentVisualizer.h b/Engine/Source/Editor/ComponentVisualizers/Public/SplineComponentVisualizer.h
--- a/Engine/Source/Editor/ComponentVisualizers/Public/SplineComponentVisualizer.h
+++ b/Engine/Source/Editor/ComponentVisualizers/Public/SplineComponentVisualizer.h
@@ -7,6 +7,7 @@
 #include "HitProxies.h"
 #include "ComponentVisualizer.h"
 #include "Components/SplineComponent.h"
+#include "SplineGeneratorPanel.h"
 
 class AActor;
 class FEditorViewportClient;
@@ -17,6 +18,7 @@ class FUICommandList;
 class FViewport;
 class SWidget;
 class USplineComponent;
+class SSplineGeneratorPanel;
 struct FViewportClick;
 struct FConvexVolume;
 
@@ -201,6 +203,8 @@ protected:
 	/** Generate the submenu containing the lock axis types */
 	void GenerateLockAxisSubMenu(FMenuBuilder& MenuBuilder) const;
 
+	void CreateSplineGeneratorPanel();
+
 	/** Output log commands */
 	TSharedPtr<FUICommandList> SplineComponentVisualizerActions;
 
@@ -262,4 +266,7 @@ protected:
 
 private:
 	UProperty* SplineCurvesProperty;
+
+	TSharedPtr<SSplineGeneratorPanel> SplineGeneratorPanel;
+	static TWeakPtr<SWindow> WeakExistingWindow;
 };
diff --git a/Engine/Source/Editor/ComponentVisualizers/Public/SplineGeneratorPanel.h b/Engine/Source/Editor/ComponentVisualizers/Public/SplineGeneratorPanel.h
new file mode 100644
--- /dev/null
+++ b/Engine/Source/Editor/ComponentVisualizers/Public/SplineGeneratorPanel.h
@@ -0,0 +1,285 @@
+// Copyright 1998-2019 Epic Games, Inc. All Rights Reserved.
+
+#pragma once
+
+#include "Widgets/SWindow.h"
+#include "UObject/Object.h"
+#include "Modules/ModuleManager.h"
+#include "PropertyEditorModule.h"
+#include "ClassViewerFilter.h"
+#include "ClassViewerModule.h"
+#include "SplineComponentVisualizer.h"
+#include "Widgets/DeclarativeSyntaxSupport.h"
+#include "Widgets/SCompoundWidget.h"
+#include "Templates/SharedPointer.h"
+#include "Templates/SubclassOf.h"
+#include "Misc/EnumClassFlags.h"
+
+#include "SplineGeneratorPanel.generated.h"
+
+class FSplineComponentVisualizer;
+
+UENUM()
+enum class EShapeAddMode : uint8
+{
+	AppendAfter = 0x01,
+	AppendBefore = 0x02,
+	InsertAfter = 0x04,
+	InsertBefore = 0x08
+};
+
+ENUM_CLASS_FLAGS(EShapeAddMode);
+
+UCLASS(config = EditorSettings)
+class USplineGeneratorBase : public UObject
+{
+	GENERATED_BODY()
+public:
+
+	USplineGeneratorBase()
+		: ShapeAddMode(EShapeAddMode::AppendAfter)
+	{}
+
+	void Init(TWeakPtr<FSplineComponentVisualizer> InSplineComponentVisualizer);
+	void Reset();
+	void PreviewCurve();
+
+	virtual void BuildCurve() {}
+	virtual int32 GetNumPoints() const { return 0; }
+
+	UPROPERTY(Transient, EditAnywhere, Category = GenerateOptions, Meta = (ToolTip = "How to add the shape to the selection"))
+	EShapeAddMode ShapeAddMode;
+
+	int32 StartKey;
+
+protected:
+	/** Helper method to return the index for adding a spline point */
+	int32 GetAddIndex(int32 Index) const;
+	/** Helper method to return the index for an added spline point after all points have been added */
+	int32 GetItrIndex(int32 Index) const;
+
+	FSplineCurves CachedSplineCurves;
+
+	TWeakObjectPtr<USplineComponent> SelectedSplineComponent;
+	TWeakPtr<FSplineComponentVisualizer> WeakSplineVis;
+};
+
+UCLASS(config = EditorSettings, DisplayName = "Circle")
+class UCircleSplineGenerator : public USplineGeneratorBase
+{
+	GENERATED_BODY()
+
+public:
+	UCircleSplineGenerator()
+		: NumberOfPoints(3)
+		, Radius(100.f)
+		, bReverseDir(false)
+		, bKeepFirstKeyTangent(true)
+		, bBranchRight(false)
+	{}
+
+	virtual void BuildCurve() override;
+
+	virtual int32 GetNumPoints() const override { return NumberOfPoints; }
+
+	UPROPERTY(Transient, EditAnywhere, Category = ShapeParameters, Meta = (ToolTip = "Number of points making up shape", ClampMin = "2"))
+	int32 NumberOfPoints;
+
+	UPROPERTY(Transient, EditAnywhere, Category = ShapeParameters, Meta=(ToolTip="Radius of circle", ClampMin = "0"))
+	float Radius;
+
+	UPROPERTY(Transient, EditAnywhere, Category = OtherParameters, Meta=(ToolTip="If enabled, will reverse the direction of the circle"))
+	bool bReverseDir;
+
+	UPROPERTY(Transient, EditAnywhere, Category = OtherParameters, Meta=(ToolTip="If enabled, will start the shape tangent to the current path"))
+	bool bKeepFirstKeyTangent;
+
+	UPROPERTY(Transient, EditAnywhere, Category = OtherParameters, Meta=(ToolTip = "If enabled, will switch the circle's center to the right of the curve", EditCondition="bKeepFirstKeyTangent"))
+	bool bBranchRight;
+};
+
+UCLASS(config = EditorSettings, DisplayName = "Arc")
+class UArcSplineGenerator : public USplineGeneratorBase
+{
+	GENERATED_BODY()
+
+public:
+	UArcSplineGenerator()
+		: NumberOfPoints(4)
+		, Radius(100.f)
+		, Degrees(90.f)
+		, bReverseDir(false)
+		, bKeepFirstKeyTangent(true)
+		, bBranchRight(false)
+	{}
+
+	virtual void BuildCurve() override;
+
+	virtual int32 GetNumPoints() const override { return NumberOfPoints; }
+
+	UPROPERTY(Transient, EditAnywhere, Category = ShapeParameters, Meta = (ToolTip = "Number of points making up shape", ClampMin = "2"))
+	int32 NumberOfPoints;
+
+	UPROPERTY(Transient, EditAnywhere, Category = ShapeParameters, Meta = (ToolTip = "Radius of arc", ClampMin = "0"))
+	float Radius;
+
+	UPROPERTY(Transient, EditAnywhere, Category = ShapeParameters, Meta = (ToolTip = "Degree of arc", ClampMin = "0"))
+	float Degrees;
+
+	UPROPERTY(Transient, EditAnywhere, Category = OtherParameters, Meta = (ToolTip = "If enabled, will reverse the direction of the arc"))
+	bool bReverseDir;
+
+	UPROPERTY(Transient, EditAnywhere, Category = OtherParameters, Meta = (ToolTip = "If enabled, will start the shape tangent to the current path"))
+	bool bKeepFirstKeyTangent;
+
+	UPROPERTY(Transient, EditAnywhere, Category = OtherParameters, Meta = (ToolTip = "If enabled, will switch the arc's center to the right of the curve", EditCondition = "bKeepFirstKeyTangent"))
+	bool bBranchRight;
+};
+
+UCLASS(config = EditorSettings, DisplayName = "Square")
+class USquareSplineGenerator : public USplineGeneratorBase
+{
+	GENERATED_BODY()
+
+public:
+	USquareSplineGenerator()
+		: Length(100.f)
+		, bBranchRight(false)
+	{}
+
+	virtual void BuildCurve() override;
+
+	virtual int32 GetNumPoints() const override { return 4; }
+
+	UPROPERTY(Transient, EditAnywhere, Category = ShapeParameters, Meta = (ClampMin = "0"))
+	float Length;
+
+	UPROPERTY(Transient, EditAnywhere, Category = OtherParameters, Meta = (ToolTip = "If enabled, will switch the shape's center to the right of the curve"))
+	bool bBranchRight;
+};
+
+UCLASS(config = EditorSettings, DisplayName = "Ellipse")
+class UEllipseSplineGenerator : public USplineGeneratorBase
+{
+	GENERATED_BODY()
+
+public:
+	UEllipseSplineGenerator()
+		: NumberOfPoints(10)
+		, Length(200.f)
+		, Width(200.f)
+		, bReverseDir(false)
+		, bKeepFirstKeyTangent(true)
+		, bBranchRight(false)
+	{}
+
+	virtual void BuildCurve() override;
+
+	UPROPERTY(Transient, EditAnywhere, Category = ShapeParameters, Meta = (ClampMin = "2"))
+	int32 NumberOfPoints;
+
+	UPROPERTY(Transient, EditAnywhere, Category = ShapeParameters, Meta = (ClampMin = "0"))
+	float Length;
+
+	UPROPERTY(Transient, EditAnywhere, Category = ShapeParameters, Meta = (ClampMin = "0"))
+	float Width;
+
+	UPROPERTY(Transient, EditAnywhere, Category = OtherParameters, Meta = (ToolTip = "If enabled, will reverse the direction of the arc"))
+	bool bReverseDir;
+
+	UPROPERTY(Transient, EditAnywhere, Category = OtherParameters, Meta = (ToolTip = "If enabled, will start the shape tangent to the current path"))
+	bool bKeepFirstKeyTangent;
+
+	UPROPERTY(Transient, EditAnywhere, Category = OtherParameters, Meta = (ToolTip = "If enabled, will switch the shape's center to the right of the curve", EditCondition = "bKeepFirstKeyTangent"))
+	bool bBranchRight;
+};
+
+UCLASS(config = EditorSettings, DisplayName = "Rectangle")
+class URectangleSplineGenerator : public USplineGeneratorBase
+{
+	GENERATED_BODY()
+
+public:
+	URectangleSplineGenerator()
+		: Length(100.f)
+		, Width(100.f)
+		, bBranchRight(false)
+	{}
+
+	virtual void BuildCurve() override;
+
+	virtual int32 GetNumPoints() const override { return 4; }
+
+	UPROPERTY(Transient, EditAnywhere, Category = ShapeParameters, Meta = (ClampMin = "0"))
+	float Length;
+
+	UPROPERTY(Transient, EditAnywhere, Category = ShapeParameters, Meta = (ClampMin = "0"))
+	float Width;
+
+	UPROPERTY(Transient, EditAnywhere, Category = OtherParameters, Meta = (ToolTip = "If enabled, will switch the shape's center to the right of the curve"))
+	bool bBranchRight;
+};
+
+UCLASS(config = EditorSettings, DisplayName = "Line")
+class ULineSplineGenerator : public USplineGeneratorBase
+{
+	GENERATED_BODY()
+
+public:
+	ULineSplineGenerator()
+		: NumberOfPoints(5)
+		, Length(100.f)
+		, bEnableUpToNextPoint(false)
+		, bUpToNextPoint(false)
+	{}
+
+	virtual void BuildCurve() override;
+
+	virtual int32 GetNumPoints() const override { return NumberOfPoints; };
+
+	UPROPERTY(Transient, EditAnywhere, Category = ShapeParameters, Meta = (ClampMin = "1"))
+	int32 NumberOfPoints;
+
+	UPROPERTY(Transient, EditAnywhere, Category = ShapeParameters, Meta = (ClampMin = "0", EditCondition = "!bUpUntilNextPoint"))
+	float Length;
+
+	UPROPERTY()
+	bool bEnableUpToNextPoint;
+
+	UPROPERTY(Transient, EditAnywhere, Category = ShapeParameters, Meta = (ToolTip = "If enabled, will add points up until the next existing point", EditCondition = "bEnableUpToNextPoint"))
+	bool bUpToNextPoint;
+
+};
+
+class SSplineGeneratorPanel : public SCompoundWidget
+{
+public:
+	SLATE_BEGIN_ARGS(SSplineGeneratorPanel)
+	{}
+
+	SLATE_END_ARGS()
+
+	/** SWidget Interface */
+	void Construct(const FArguments& InArgs, TWeakPtr<FSplineComponentVisualizer> InWeakSplineComponentVisualizer);
+	virtual void Tick(const FGeometry& AllottedGeometry, const double InCurrentTime, const float InDeltaTime) override;
+	virtual void OnFocusChanging(const FWeakWidgetPath& PreviousFocusPath, const FWidgetPath& NewWidgetPath, const FFocusEvent& InFocusEvent) override;
+
+	/** SWindow Interface */
+	void OnWindowClosed(const TSharedRef<SWindow>&);
+
+	void OnSelectionUpdated();
+
+	~SSplineGeneratorPanel();
+
+private:
+
+	TSharedPtr<IDetailsView> DetailView;
+
+	TWeakPtr<FSplineComponentVisualizer> WeakSplineComponentVisualizer;
+	USplineGeneratorBase* SplineGen;
+
+	TArray<UObject*> ShapeGenRegistry;
+
+	/** The currently open transaction (if any) */
+	TUniquePtr<class FScopedTransaction> ActiveTransaction;
+};
\ No newline at end of file
diff --git a/Engine/Source/Runtime/CinematicCamera/Private/CameraRig_Rail.cpp b/Engine/Source/Runtime/CinematicCamera/Private/CameraRig_Rail.cpp
--- a/Engine/Source/Runtime/CinematicCamera/Private/CameraRig_Rail.cpp
+++ b/Engine/Source/Runtime/CinematicCamera/Private/CameraRig_Rail.cpp
@@ -19,6 +19,11 @@ ACameraRig_Rail::ACameraRig_Rail(const FObjectInitializer& ObjectInitializer)
 	CurrentPositionOnRail = 0.f;
 	bLockOrientationToRail = false;
 
+#if WITH_EDITORONLY_DATA
+	bShowRailVisualization = true;
+	PreviewMeshScale = 1.f;
+#endif
+
 	// Create components
 	TransformComponent = CreateDefaultSubobject<USceneComponent>(TEXT("TransformComponent"));
 
@@ -120,6 +125,8 @@ void ACameraRig_Rail::UpdatePreviewMeshes()
 				USplineMeshComponent* const SplineMeshComp = PreviewRailMeshSegments[PtIdx];
 				if (SplineMeshComp)
 				{
+					SplineMeshComp->SetVisibility(bShowRailVisualization);
+					SplineMeshComp->SetWorldScale3D(FVector(1.f, PreviewMeshScale, PreviewMeshScale));
 					SplineMeshComp->SetForwardAxis(ESplineMeshAxis::Z);
 					SplineMeshComp->SetStartAndEnd(StartLoc, StartTangent, EndLoc, EndTangent, true);
 				}
@@ -136,13 +143,15 @@ void ACameraRig_Rail::UpdatePreviewMeshes()
 				}
 			}
 		}
-
+		
 		// make visualization of the mount follow the contour of the rail
 		if (PreviewMesh_Mount)
 		{
 			float const SplineLen = RailSplineComponent->GetSplineLength();
 			FQuat const RailRot = RailSplineComponent->GetQuaternionAtDistanceAlongSpline(CurrentPositionOnRail*SplineLen, ESplineCoordinateSpace::World);
 			PreviewMesh_Mount->SetWorldRotation(RailRot);
+			PreviewMesh_Mount->SetVisibility(bShowRailVisualization);
+			PreviewMesh_Mount->SetWorldScale3D(FVector(PreviewMeshScale, PreviewMeshScale, PreviewMeshScale));
 		}
 	}
 }
@@ -189,6 +198,12 @@ void ACameraRig_Rail::PostEditUndo()
 	Super::PostEditUndo();
 	UpdateRailComponents();
 }
+
+void ACameraRig_Rail::PostEditMove(bool bFinished)
+{
+	Super::PostEditMove(bFinished);
+	UpdateRailComponents();
+}
 #endif
 
 void ACameraRig_Rail::Tick(float DeltaTime)
diff --git a/Engine/Source/Runtime/CinematicCamera/Public/CameraRig_Rail.h b/Engine/Source/Runtime/CinematicCamera/Public/CameraRig_Rail.h
--- a/Engine/Source/Runtime/CinematicCamera/Public/CameraRig_Rail.h
+++ b/Engine/Source/Runtime/CinematicCamera/Public/CameraRig_Rail.h
@@ -35,10 +35,21 @@ public:
 	UPROPERTY(Interp, EditAnywhere, BlueprintReadWrite, Category = "Rail Controls")
 	bool bLockOrientationToRail;
 
+#if WITH_EDITORONLY_DATA
+	/** Determines whether or not to show the rail mesh preview. */
+	UPROPERTY(Transient, EditAnywhere, Category = "Rail Controls")
+	bool bShowRailVisualization;
+
+	/** Determines the scale of the rail mesh preview */
+	UPROPERTY(Transient, EditAnywhere, Category = "Rail Controls")
+	float PreviewMeshScale;
+#endif
+
 	virtual class USceneComponent* GetDefaultAttachComponent() const override;
 #if WITH_EDITOR
 	virtual void PostEditChangeProperty(FPropertyChangedEvent& PropertyChangedEvent) override;
 	virtual void PostEditUndo() override;
+	virtual void PostEditMove(bool bFinished) override;
 #endif
 
 	/** Returns the spline component that defines the rail path */
