Tools: Fix compile warning in SplineGeneratorPanel.cpp

diff --git a/Engine/Source/Editor/ComponentVisualizers/Private/SplineGeneratorPanel.cpp b/Engine/Source/Editor/ComponentVisualizers/Private/SplineGeneratorPanel.cpp
--- a/Engine/Source/Editor/ComponentVisualizers/Private/SplineGeneratorPanel.cpp
+++ b/Engine/Source/Editor/ComponentVisualizers/Private/SplineGeneratorPanel.cpp
@@ -615,4 +615,7 @@ SSplineGeneratorPanel::~SSplineGeneratorPanel()
 	{
 		ActiveTransaction ->Cancel();
 	}
-}
\ No newline at end of file
+}
+
+
+#undef LOCTEXT_NAMESPACE
\ No newline at end of file
