Tools: Add missing include in SplineGeneratorPanel.cpp

diff --git a/Engine/Source/Editor/ComponentVisualizers/Private/SplineGeneratorPanel.cpp b/Engine/Source/Editor/ComponentVisualizers/Private/SplineGeneratorPanel.cpp
--- a/Engine/Source/Editor/ComponentVisualizers/Private/SplineGeneratorPanel.cpp
+++ b/Engine/Source/Editor/ComponentVisualizers/Private/SplineGeneratorPanel.cpp
@@ -6,6 +6,7 @@
 #include "Widgets/Input/SCheckBox.h"
 #include "Widgets/Text/STextBlock.h"
 #include "ScopedTransaction.h"
+#include "Editor.h"
 
 #define LOCTEXT_NAMESPACE "SplineGenerator"
 
