Fix Mac Compile Warnings

#rnx

diff --git a/Engine/Source/Editor/ComponentVisualizers/Private/SplineComponentVisualizer.cpp b/Engine/Source/Editor/ComponentVisualizers/Private/SplineComponentVisualizer.cpp
--- a/Engine/Source/Editor/ComponentVisualizers/Private/SplineComponentVisualizer.cpp
+++ b/Engine/Source/Editor/ComponentVisualizers/Private/SplineComponentVisualizer.cpp
@@ -2417,7 +2417,7 @@ TSharedPtr<SWidget> FSplineComponentVisualizer::GenerateContextMenu() const
 				LOCTEXT("SplineGenerateTooltip", "Opens up a spline generation panel to easily create basic shapes with splines"),
 				FSlateIcon(),
 				FUIAction( 
-					FExecuteAction::CreateRaw(this, &FSplineComponentVisualizer::CreateSplineGeneratorPanel), 
+					FExecuteAction::CreateSP(this, &FSplineComponentVisualizer::CreateSplineGeneratorPanel), 
 					FCanExecuteAction::CreateLambda([] { return true; })
 				)
 			);
diff --git a/Engine/Source/Editor/FoliageEdit/Private/SFoliageEdit.cpp b/Engine/Source/Editor/FoliageEdit/Private/SFoliageEdit.cpp
--- a/Engine/Source/Editor/FoliageEdit/Private/SFoliageEdit.cpp
+++ b/Engine/Source/Editor/FoliageEdit/Private/SFoliageEdit.cpp
@@ -1058,7 +1058,7 @@ void SFoliageEdit::OnSingleInstantiationPlacementModeChanged(int32 InMode)
 	FoliageEditMode->UISettings.SetSingleInstantiationPlacementMode(EFoliageSingleInstantiationPlacementMode::Type(InMode));
 }
 
-TSharedRef<SWidget> SFoliageEdit::GetSingleInstantiationModeMenuContent() const
+TSharedRef<SWidget> SFoliageEdit::GetSingleInstantiationModeMenuContent()
 {
 	FMenuBuilder MenuBuilder(true, nullptr);
 
diff --git a/Engine/Source/Editor/FoliageEdit/Private/SFoliageEdit.h b/Engine/Source/Editor/FoliageEdit/Private/SFoliageEdit.h
--- a/Engine/Source/Editor/FoliageEdit/Private/SFoliageEdit.h
+++ b/Engine/Source/Editor/FoliageEdit/Private/SFoliageEdit.h
@@ -199,7 +199,7 @@ private:	// BRUSH SETTINGS
 	FText GetSingleInstantiationPlacementModeText(EFoliageSingleInstantiationPlacementMode::Type InMode) const;
 
 	/** @return menu of all single instantiation modes */
-	TSharedRef<SWidget> GetSingleInstantiationModeMenuContent() const;
+	TSharedRef<SWidget> GetSingleInstantiationModeMenuContent();
 
 	/** @return display text for current placement type */
 	FText GetCurrentSingleInstantiationPlacementModeText() const;
