Spline Component: Fix const delegate deprecation warning

diff --git a/Engine/Source/Editor/ComponentVisualizers/Private/SplineComponentVisualizer.cpp b/Engine/Source/Editor/ComponentVisualizers/Private/SplineComponentVisualizer.cpp
--- a/Engine/Source/Editor/ComponentVisualizers/Private/SplineComponentVisualizer.cpp
+++ b/Engine/Source/Editor/ComponentVisualizers/Private/SplineComponentVisualizer.cpp
@@ -2417,7 +2417,7 @@ TSharedPtr<SWidget> FSplineComponentVisualizer::GenerateContextMenu() const
 				LOCTEXT("SplineGenerateTooltip", "Opens up a spline generation panel to easily create basic shapes with splines"),
 				FSlateIcon(),
 				FUIAction( 
-					FExecuteAction::CreateSP(this, &FSplineComponentVisualizer::CreateSplineGeneratorPanel), 
+					FExecuteAction::CreateSP(const_cast<FSplineComponentVisualizer*>(this), &FSplineComponentVisualizer::CreateSplineGeneratorPanel),
 					FCanExecuteAction::CreateLambda([] { return true; })
 				)
 			);
