Fix include cycle in SplineGeneratorPanel.h

diff --git a/Engine/Source/Editor/ComponentVisualizers/Public/SplineComponentVisualizer.h b/Engine/Source/Editor/ComponentVisualizers/Public/SplineComponentVisualizer.h
--- a/Engine/Source/Editor/ComponentVisualizers/Public/SplineComponentVisualizer.h
+++ b/Engine/Source/Editor/ComponentVisualizers/Public/SplineComponentVisualizer.h
@@ -7,7 +7,6 @@
 #include "HitProxies.h"
 #include "ComponentVisualizer.h"
 #include "Components/SplineComponent.h"
-#include "SplineGeneratorPanel.h"
 
 class AActor;
 class FEditorViewportClient;
diff --git a/Engine/Source/Editor/ComponentVisualizers/Public/SplineGeneratorPanel.h b/Engine/Source/Editor/ComponentVisualizers/Public/SplineGeneratorPanel.h
--- a/Engine/Source/Editor/ComponentVisualizers/Public/SplineGeneratorPanel.h
+++ b/Engine/Source/Editor/ComponentVisualizers/Public/SplineGeneratorPanel.h
@@ -8,16 +8,14 @@
 #include "PropertyEditorModule.h"
 #include "ClassViewerFilter.h"
 #include "ClassViewerModule.h"
-#include "SplineComponentVisualizer.h"
 #include "Widgets/DeclarativeSyntaxSupport.h"
 #include "Widgets/SCompoundWidget.h"
 #include "Templates/SharedPointer.h"
-#include "Templates/SubclassOf.h"
 #include "Misc/EnumClassFlags.h"
+#include "SplineComponentVisualizer.h"
 
 #include "SplineGeneratorPanel.generated.h"
 
-class FSplineComponentVisualizer;
 
 UENUM()
 enum class EShapeAddMode : uint8
