Add missing include in SplineComponentVisualizer.cpp

diff --git a/Engine/Source/Editor/ComponentVisualizers/Private/SplineComponentVisualizer.cpp b/Engine/Source/Editor/ComponentVisualizers/Private/SplineComponentVisualizer.cpp
--- a/Engine/Source/Editor/ComponentVisualizers/Private/SplineComponentVisualizer.cpp
+++ b/Engine/Source/Editor/ComponentVisualizers/Private/SplineComponentVisualizer.cpp
@@ -18,6 +18,7 @@
 #include "ActorEditorUtils.h"
 #include "WorldCollision.h"
 #include "Widgets/Docking/SDockTab.h"
+#include "SplineGeneratorPanel.h"
 
 IMPLEMENT_HIT_PROXY(HSplineVisProxy, HComponentVisProxy);
 IMPLEMENT_HIT_PROXY(HSplineKeyProxy, HSplineVisProxy);
