Fixed improper handling of closed-loop splines, added more checks.

diff --git a/Engine/Source/Editor/ComponentVisualizers/Private/SplineComponentVisualizer.cpp b/Engine/Source/Editor/ComponentVisualizers/Private/SplineComponentVisualizer.cpp
--- a/Engine/Source/Editor/ComponentVisualizers/Private/SplineComponentVisualizer.cpp
+++ b/Engine/Source/Editor/ComponentVisualizers/Private/SplineComponentVisualizer.cpp
@@ -347,6 +347,9 @@ void FSplineComponentVisualizer::DrawVisualization(const UActorComponent* Compon
 		{
 			for (int32 SelectedKey : SelectedKeys)
 			{
+				check(SelectedKey >= 0);
+				check(SelectedKey < SplineComp->GetNumberOfSplinePoints());
+
 				if (SplineInfo.Points[SelectedKey].IsCurveKey())
 				{
 					const FVector Location = SplineComp->GetLocationAtSplinePoint(SelectedKey, ESplineCoordinateSpace::World);
@@ -697,6 +700,7 @@ bool FSplineComponentVisualizer::GetWidgetLocation(const FEditorViewportClient*
 		else if (LastKeyIndexSelected != INDEX_NONE)
 		{
 			// Otherwise use the last key index set
+			check(LastKeyIndexSelected >= 0);
 			check(LastKeyIndexSelected < Position.Points.Num());
 			check(SelectedKeys.Contains(LastKeyIndexSelected));
 			const FInterpCurvePointVector& Point = Position.Points[LastKeyIndexSelected];
@@ -853,13 +857,18 @@ bool FSplineComponentVisualizer::TransformSelectedKeys(const FVector& DeltaTrans
 		const int32 NumPoints = SplinePosition.Points.Num();
 
 		check(LastKeyIndexSelected != INDEX_NONE);
+		check(LastKeyIndexSelected >= 0);
 		check(LastKeyIndexSelected < NumPoints);
 		check(SelectedKeys.Num() > 0);
+		check(SelectedKeys.Contains(LastKeyIndexSelected));
 
 		SplineComp->Modify();
 
 		for (int32 SelectedKeyIndex : SelectedKeys)
 		{
+			check(SelectedKeyIndex >= 0); 
+			check(SelectedKeyIndex < NumPoints);
+
 			FInterpCurvePoint<FVector>& EditedPoint = SplinePosition.Points[SelectedKeyIndex];
 			FInterpCurvePoint<FQuat>& EditedRotPoint = SplineRotation.Points[SelectedKeyIndex];
 			FInterpCurvePoint<FVector>& EditedScalePoint = SplineScale.Points[SelectedKeyIndex];
@@ -1103,6 +1112,9 @@ bool FSplineComponentVisualizer::HasFocusOnSelectionBoundingBox(FBox& OutBoundin
 			// Spline control point selection always uses transparent box selection.
 			for (int32 KeyIdx : SelectedKeys)
 			{
+				check(KeyIdx >= 0); 
+				check(KeyIdx < SplineComp->GetNumberOfSplinePoints());
+
 				const FVector Pos = SplineComp->GetLocationAtSplinePoint(KeyIdx, ESplineCoordinateSpace::World);
 
 				OutBoundingBox += Pos;
@@ -1133,6 +1145,8 @@ bool FSplineComponentVisualizer::HandleSnapTo(const bool bInAlign, const bool bI
 		if (SelectedKeys.Num() > 0)
 		{
 			check(LastKeyIndexSelected != INDEX_NONE);
+			check(LastKeyIndexSelected >= 0);
+			check(LastKeyIndexSelected < SplineComp->GetNumberOfSplinePoints());
 			check(SelectedKeys.Contains(LastKeyIndexSelected));
 
 			SplineComp->Modify();
@@ -1146,6 +1160,7 @@ bool FSplineComponentVisualizer::HandleSnapTo(const bool bInAlign, const bool bI
 			// Spline control point selection always uses transparent box selection.
 			for (int32 KeyIdx : SelectedKeys)
 			{
+				check(KeyIdx >= 0);
 				check(KeyIdx < NumPoints);
 
 				FVector Direction = FVector(0.f, 0.f, -1.f);
@@ -1220,6 +1235,8 @@ void FSplineComponentVisualizer::OnSnapToNearestSplinePoint(ESplineComponentSnap
 	USplineComponent* SplineComp = GetEditedSplineComponent();
 	check(SplineComp != nullptr);
 	check(LastKeyIndexSelected != INDEX_NONE);
+	check(LastKeyIndexSelected >= 0);
+	check(LastKeyIndexSelected < SplineComp->GetNumberOfSplinePoints());
 	check(SelectedKeys.Num() == 1);
 	check(SelectedKeys.Contains(LastKeyIndexSelected));
 
@@ -1379,6 +1396,8 @@ void FSplineComponentVisualizer::OnSnapAll(EAxis::Type InAxis)
 	USplineComponent* SplineComp = GetEditedSplineComponent();
 	check(SplineComp != nullptr);
 	check(LastKeyIndexSelected != INDEX_NONE);
+	check(LastKeyIndexSelected >= 0);
+	check(LastKeyIndexSelected < SplineComp->GetNumberOfSplinePoints());
 	check(SelectedKeys.Num() == 1);
 	check(SelectedKeys.Contains(LastKeyIndexSelected));
 	check(InAxis == EAxis::X || InAxis == EAxis::Y || InAxis == EAxis::Z);
@@ -1493,6 +1512,8 @@ void FSplineComponentVisualizer::OnDuplicateKey()
 	USplineComponent* SplineComp = GetEditedSplineComponent();
 	check(SplineComp != nullptr);
 	check(LastKeyIndexSelected != INDEX_NONE);
+	check(LastKeyIndexSelected >= 0);
+	check(LastKeyIndexSelected < SplineComp->GetNumberOfSplinePoints());
 	check(SelectedKeys.Num() > 0);
 	check(SelectedKeys.Contains(LastKeyIndexSelected));
 
@@ -1518,6 +1539,9 @@ void FSplineComponentVisualizer::OnDuplicateKey()
 
 	for (int32 SelectedKeyIndex : SelectedKeysSorted)
 	{
+		check(SelectedKeyIndex >= 0);
+		check(SelectedKeyIndex < SplineComp->GetNumberOfSplinePoints());
+
 		// Insert duplicates into arrays.
 		// It's necessary to take a copy because copying existing array items by reference isn't allowed (the array may reallocate)
 		SplinePosition.Points.Insert(FInterpCurvePoint<FVector>(SplinePosition.Points[SelectedKeyIndex]), SelectedKeyIndex);
@@ -1578,10 +1602,7 @@ bool FSplineComponentVisualizer::CanAddKeyToSegment() const
 		return false;
 	}
 
-	const int32 NumPoints = SplineComp->SplineCurves.Position.Points.Num();
-	const int32 NumSegments = SplineComp->IsClosedLoop() ? NumPoints : NumPoints - 1;
-
-	return (SelectedSegmentIndex != INDEX_NONE && SelectedSegmentIndex < NumSegments);
+	return (SelectedSegmentIndex != INDEX_NONE && SelectedSegmentIndex >=0 && SelectedSegmentIndex < SplineComp->GetNumberOfSplineSegments());
 }
 
 void FSplineComponentVisualizer::OnAddKeyToSegment()
@@ -1590,8 +1611,8 @@ void FSplineComponentVisualizer::OnAddKeyToSegment()
 	USplineComponent* SplineComp = GetEditedSplineComponent();
 	check(SplineComp != nullptr);
 	check(LastKeyIndexSelected != INDEX_NONE);
-	check(SelectedKeys.Num() > 0);
-	check(SelectedKeys.Contains(LastKeyIndexSelected));
+	check(LastKeyIndexSelected >= 0);
+	check(LastKeyIndexSelected < SplineComp->GetNumberOfSplinePoints());
 	check(SelectedTangentHandle == INDEX_NONE);
 	check(SelectedTangentHandleType == ESelectedTangentHandle::None);
 
@@ -1607,7 +1628,10 @@ bool FSplineComponentVisualizer::DuplicateKeyForAltDrag(const FVector& InDrag)
 {
 	USplineComponent* SplineComp = GetEditedSplineComponent();
 	check(SplineComp != nullptr);
+	const int32 NumPoints = SplineComp->GetNumberOfSplinePoints();
 	check(LastKeyIndexSelected != INDEX_NONE);
+	check(LastKeyIndexSelected >= 0);
+	check(LastKeyIndexSelected < NumPoints);
 	check(SelectedKeys.Num() == 1);
 	check(SelectedKeys.Contains(LastKeyIndexSelected));
 
@@ -1619,11 +1643,13 @@ bool FSplineComponentVisualizer::DuplicateKeyForAltDrag(const FVector& InDrag)
 	const FVector CurrentKeyWorldPos = SplineComp->GetComponentTransform().TransformPosition(SplinePosition.Points[CurrentIndex].OutVal);
 
 	// Determine direction to insert new point				
-	bool bHasPrevKey = false;
+	bool bHasPrevKey = SplineComp->IsClosedLoop() || CurrentIndex > 0;
 	float PrevAngle = 0.0f;
-	if (CurrentIndex > 0)
+	if (bHasPrevKey)
 	{
-		FVector PrevKeyWorldPos = SplineComp->GetComponentTransform().TransformPosition(SplinePosition.Points[CurrentIndex - 1].OutVal);
+		// Wrap index around for closed-looped splines
+		int32 PrevKeyIndex = (CurrentIndex > 0 ? CurrentIndex - 1 : NumPoints - 1);
+		FVector PrevKeyWorldPos = SplineComp->GetComponentTransform().TransformPosition(SplinePosition.Points[PrevKeyIndex].OutVal);
 		FVector SegmentDirection = PrevKeyWorldPos - CurrentKeyWorldPos;
 		if (!SegmentDirection.IsZero())
 		{
@@ -1633,14 +1659,15 @@ bool FSplineComponentVisualizer::DuplicateKeyForAltDrag(const FVector& InDrag)
 		{
 			PrevAngle = HALF_PI;
 		}
-		bHasPrevKey = true;
 	}
 
-	bool bHasNextKey = false;
+	bool bHasNextKey = SplineComp->IsClosedLoop() || CurrentIndex + 1 < NumPoints;
 	float NextAngle = 0.0f;
-	if (CurrentIndex + 1 < SplinePosition.Points.Num())
+	if (bHasNextKey)
 	{
-		FVector NextKeyWorldPos = SplineComp->GetComponentTransform().TransformPosition(SplinePosition.Points[CurrentIndex + 1].OutVal);
+		// Wrap index around for closed-looped splines
+		int32 NextKeyIndex = (CurrentIndex + 1 < NumPoints ? CurrentIndex + 1 : 0);
+		FVector NextKeyWorldPos = SplineComp->GetComponentTransform().TransformPosition(SplinePosition.Points[NextKeyIndex].OutVal);
 		FVector SegmentDirection = NextKeyWorldPos - CurrentKeyWorldPos;
 		if (!SegmentDirection.IsZero())
 		{
@@ -1650,7 +1677,6 @@ bool FSplineComponentVisualizer::DuplicateKeyForAltDrag(const FVector& InDrag)
 		{
 			NextAngle = HALF_PI;
 		}
-		bHasNextKey = true;
 	}
 
 	// Set key index to which the drag will be applied after duplication
@@ -1663,10 +1689,17 @@ bool FSplineComponentVisualizer::DuplicateKeyForAltDrag(const FVector& InDrag)
 		SegmentIndex--;
 	}
 
+	// Wrap index around for closed-looped splines
+	const int32 NumSegments = SplineComp->GetNumberOfSplineSegments();
+	if (SplineComp->IsClosedLoop() && SegmentIndex < 0)
+	{
+		SegmentIndex = NumSegments - 1;
+	}
+
 	FVector WorldPos = CurrentKeyWorldPos + InDrag;
 
 	// Split existing segment or add new segment
-	if (SegmentIndex >= 0 && SegmentIndex < SplinePosition.Points.Num() - 1)
+	if (SegmentIndex >= 0 && SegmentIndex < NumSegments)
 	{
 		SplitSegment(WorldPos, SegmentIndex);
 	}
@@ -1701,8 +1734,9 @@ float FSplineComponentVisualizer::FindNearest(const FVector& InLocalPos, int32 I
 {
 	USplineComponent* SplineComp = GetEditedSplineComponent();
 	check(SplineComp != nullptr);
+	check(InSegmentIndex != INDEX_NONE);
 	check(InSegmentIndex >= 0);
-	check(InSegmentIndex < SplineComp->GetSplinePointsPosition().Points.Num() - 1);
+	check(InSegmentIndex < SplineComp->GetNumberOfSplineSegments());
 
 	FInterpCurveVector& SplinePosition = SplineComp->GetSplinePointsPosition();
 	float OutSquaredDistance = 0.0f;
@@ -1718,10 +1752,13 @@ void FSplineComponentVisualizer::SplitSegment(const FVector& InWorldPos, int32 I
 	USplineComponent* SplineComp = GetEditedSplineComponent();
 	check(SplineComp != nullptr);
 	check(LastKeyIndexSelected != INDEX_NONE);
+	check(LastKeyIndexSelected >= 0);
+	check(LastKeyIndexSelected < SplineComp->GetNumberOfSplinePoints());
 	check(SelectedKeys.Num() == 1);
 	check(SelectedKeys.Contains(LastKeyIndexSelected));
+	check(InSegmentIndex != INDEX_NONE);
 	check(InSegmentIndex >= 0);
-	check(InSegmentIndex < SplineComp->GetSplinePointsPosition().Points.Num() - 1);
+	check(InSegmentIndex < SplineComp->GetNumberOfSplineSegments());
 
 	SplineComp->Modify();
 	if (AActor* Owner = SplineComp->GetOwner())
@@ -1743,6 +1780,11 @@ void FSplineComponentVisualizer::SplitSegment(const FVector& InWorldPos, int32 I
 
 	int32 SegmentBeginIndex = InSegmentIndex;
 	int32 SegmentSplitIndex = InSegmentIndex + 1;
+	int32 SegmentEndIndex = SegmentSplitIndex;
+	if (SplineComp->IsClosedLoop() && SegmentEndIndex >= SplineComp->GetNumberOfSplinePoints())
+	{
+		SegmentEndIndex = 0;
+	}
 
 	FInterpCurveVector& SplinePosition = SplineComp->GetSplinePointsPosition();
 	FInterpCurveQuat& SplineRotation = SplineComp->GetSplinePointsRotation();
@@ -1754,21 +1796,21 @@ void FSplineComponentVisualizer::SplitSegment(const FVector& InWorldPos, int32 I
 	{
 		SplinePosition.Points[SegmentBeginIndex].InterpMode = CIM_CurveAuto;
 	}
-	if (SplinePosition.Points[SegmentSplitIndex].InterpMode == CIM_CurveUser)
+	if (SplinePosition.Points[SegmentEndIndex].InterpMode == CIM_CurveUser)
 	{
-		SplinePosition.Points[SegmentSplitIndex].InterpMode = CIM_CurveAuto;
+		SplinePosition.Points[SegmentEndIndex].InterpMode = CIM_CurveAuto;
 	}
 
 	// Compute interpolated scale
 	FVector NewScale;
 	FInterpCurvePoint<FVector>& PrevScale = SplineScale.Points[SegmentBeginIndex];
-	FInterpCurvePoint<FVector>& NextScale = SplineScale.Points[SegmentSplitIndex];
+	FInterpCurvePoint<FVector>& NextScale = SplineScale.Points[SegmentEndIndex];
 	NewScale = FMath::LerpStable(PrevScale.OutVal, NextScale.OutVal, t);
 
 	// Compute interpolated rot
 	FQuat NewRot;
 	FInterpCurvePoint<FQuat>& PrevRot = SplineRotation.Points[SegmentBeginIndex];
-	FInterpCurvePoint<FQuat>& NextRot = SplineRotation.Points[SegmentSplitIndex];
+	FInterpCurvePoint<FQuat>& NextRot = SplineRotation.Points[SegmentEndIndex];
 	NewRot = FMath::Lerp(PrevRot.OutVal, NextRot.OutVal, t);
 
 	FInterpCurvePoint<FVector> NewPoint(
@@ -1792,16 +1834,27 @@ void FSplineComponentVisualizer::SplitSegment(const FVector& InWorldPos, int32 I
 		FVector::ZeroVector,
 		CIM_CurveAuto);
 
-	SplinePosition.Points.Insert(NewPoint, SegmentSplitIndex);
-	SplineRotation.Points.Insert(NewRotPoint, SegmentSplitIndex);
-	SplineScale.Points.Insert(NewScalePoint, SegmentSplitIndex);
+	if (SegmentEndIndex == 0)
+	{
+		// Splitting last segment of a closed-looped spline
+		SplinePosition.Points.Emplace(NewPoint);
+		SplineRotation.Points.Emplace(NewRotPoint);
+		SplineScale.Points.Emplace(NewScalePoint);
+	}
+	else
+	{
+		SplinePosition.Points.Insert(NewPoint, SegmentEndIndex);
+		SplineRotation.Points.Insert(NewRotPoint, SegmentEndIndex);
+		SplineScale.Points.Insert(NewScalePoint, SegmentEndIndex);
+	}
+
 	if (SplineMetadata)
 	{
-		SplineMetadata->DuplicatePoint(SegmentSplitIndex);
+		SplineMetadata->DuplicatePoint(SegmentEndIndex);
 	}
 
 	// Adjust input keys of subsequent points
-	for (int Index = SegmentSplitIndex + 1; Index < SplinePosition.Points.Num(); Index++)
+	for (int Index = SegmentSplitIndex + 1; Index < SplineComp->GetNumberOfSplinePoints(); Index++)
 	{
 		SplinePosition.Points[Index].InVal += 1.0f;
 		SplineRotation.Points[Index].InVal += 1.0f;
@@ -1826,8 +1879,10 @@ void FSplineComponentVisualizer::UpdateSplitSegment(const FVector& InDrag)
 	check(LastKeyIndexSelected != INDEX_NONE);
 	check(SelectedKeys.Num() == 1);
 	check(SelectedKeys.Contains(LastKeyIndexSelected));
+	// LastKeyIndexSelected is the newly created point when splitting a segment with alt-drag. 
+	// Check that it is an internal point, not an end point.
 	check(LastKeyIndexSelected > 0);
-	check(LastKeyIndexSelected < SplineComp->GetSplinePointsPosition().Points.Num() - 1);
+	check(LastKeyIndexSelected < SplineComp->GetNumberOfSplineSegments());
 
 	SplineComp->Modify();
 	if (AActor* Owner = SplineComp->GetOwner())
@@ -1839,6 +1894,12 @@ void FSplineComponentVisualizer::UpdateSplitSegment(const FVector& InDrag)
 	int32 SegmentSplitIndex = LastKeyIndexSelected;
 	int32 SegmentEndIndex = LastKeyIndexSelected + 1;
 
+	// Wrap end point if on last segment of closed-looped spline
+	if (SplineComp->IsClosedLoop() && SegmentEndIndex >= SplineComp->GetNumberOfSplineSegments())
+	{
+		SegmentEndIndex = 0;
+	}
+
 	FInterpCurveVector& SplinePosition = SplineComp->GetSplinePointsPosition();
 	FInterpCurveVector& SplineScale = SplineComp->GetSplinePointsScale();
 	FInterpCurveQuat& SplineRotation = SplineComp->GetSplinePointsRotation();
@@ -1891,7 +1952,7 @@ void FSplineComponentVisualizer::UpdateSplitSegment(const FVector& InDrag)
 	// Update metadata
 	if (SplineMetadata)
 	{
-		SplineMetadata->UpdatePoint(SegmentSplitIndex, t);
+		SplineMetadata->UpdatePoint(SegmentSplitIndex, t, SplineComp->IsClosedLoop());
 	}
 
 	SplineComp->UpdateSpline();
@@ -2003,9 +2064,12 @@ void FSplineComponentVisualizer::UpdateAddSegment(const FVector& InDrag)
 	USplineComponent* SplineComp = GetEditedSplineComponent();
 	check(SplineComp != nullptr);
 	check(LastKeyIndexSelected != INDEX_NONE);
+	check(LastKeyIndexSelected >= 0);
+	check(LastKeyIndexSelected < SplineComp->GetNumberOfSplinePoints());
 	check(SelectedKeys.Num() == 1);
 	check(SelectedKeys.Contains(LastKeyIndexSelected));
-	// Only work on keys at either end of the spline 
+	// Only work on keys at either end of a non-closed-looped spline 
+	check(!SplineComp->IsClosedLoop());
 	check(LastKeyIndexSelected == 0 || LastKeyIndexSelected == SplineComp->GetSplinePointsPosition().Points.Num() - 1);
 
 	SplineComp->Modify();
@@ -2044,6 +2108,8 @@ void FSplineComponentVisualizer::OnDeleteKey()
 	USplineComponent* SplineComp = GetEditedSplineComponent();
 	check(SplineComp != nullptr);
 	check(LastKeyIndexSelected != INDEX_NONE);
+	check(LastKeyIndexSelected >= 0);
+	check(LastKeyIndexSelected < SplineComp->GetNumberOfSplinePoints());
 	check(SelectedKeys.Num() > 0);
 	check(SelectedKeys.Contains(LastKeyIndexSelected));
 
@@ -2197,6 +2263,8 @@ void FSplineComponentVisualizer::OnSetKeyType(EInterpCurveMode Mode)
 
 		for (int32 SelectedKeyIndex : SelectedKeys)
 		{
+			check(SelectedKeyIndex >= 0);
+			check(SelectedKeyIndex < SplineComp->GetNumberOfSplinePoints());
 			SplineComp->SplineCurves.Position.Points[SelectedKeyIndex].InterpMode = Mode;
 		}
 
@@ -2219,6 +2287,8 @@ bool FSplineComponentVisualizer::IsKeyTypeSet(EInterpCurveMode Mode) const
 
 		for (int32 SelectedKeyIndex : SelectedKeys)
 		{
+			check(SelectedKeyIndex >= 0);
+			check(SelectedKeyIndex < SplineComp->GetNumberOfSplinePoints());
 			const auto& SelectedPoint = SplineComp->SplineCurves.Position.Points[SelectedKeyIndex];
 			if ((Mode == CIM_CurveAuto && SelectedPoint.IsCurveKey()) || SelectedPoint.InterpMode == Mode)
 			{
@@ -2400,6 +2470,8 @@ TSharedPtr<SWidget> FSplineComponentVisualizer::GenerateContextMenu() const
 				{
 					for (int32 SelectedKeyIndex : SelectedKeys)
 					{
+						check(SelectedKeyIndex >= 0);
+						check(SelectedKeyIndex < SplineComp->GetNumberOfSplinePoints());
 						const auto& Point = SplineComp->SplineCurves.Position.Points[SelectedKeyIndex];
 						if (Point.IsCurveKey())
 						{
diff --git a/Engine/Source/Runtime/Engine/Classes/Components/SplineComponent.h b/Engine/Source/Runtime/Engine/Classes/Components/SplineComponent.h
--- a/Engine/Source/Runtime/Engine/Classes/Components/SplineComponent.h
+++ b/Engine/Source/Runtime/Engine/Classes/Components/SplineComponent.h
@@ -44,9 +44,9 @@ class ENGINE_API USplineMetadata : public UObject
 
 public:
 	/** Insert point before index, lerping metadata between previous and next key values */
-	virtual void InsertPoint(int32 Index, float t) PURE_VIRTUAL(USplineMetadata::InsertPoint, );
+	virtual void InsertPoint(int32 Index, float t, bool bClosedLoop) PURE_VIRTUAL(USplineMetadata::InsertPoint, );
 	/** Update point at index by lerping metadata between previous and next key values */
-	virtual void UpdatePoint(int32 Index, float t) PURE_VIRTUAL(USplineMetadata::UpdatePoint, );
+	virtual void UpdatePoint(int32 Index, float t, bool bClosedLoop) PURE_VIRTUAL(USplineMetadata::UpdatePoint, );
 	virtual void AddPoint(float InputKey) PURE_VIRTUAL(USplineMetadata::AddPoint, );
 	virtual void RemovePoint(int32 Index) PURE_VIRTUAL(USplineMetadata::RemovePoint, );
 	virtual void DuplicatePoint(int32 Index) PURE_VIRTUAL(USplineMetadata::DuplicatePoint, );
@@ -462,6 +462,10 @@ public:
 	UFUNCTION(BlueprintCallable, Category = Spline)
 	int32 GetNumberOfSplinePoints() const;
 
+	/** Get the number of segments that make up this spline */
+	UFUNCTION(BlueprintCallable, Category = Spline)
+	int32 GetNumberOfSplineSegments() const;
+
 	/** Get the location at spline point */
 	UFUNCTION(BlueprintCallable, Category = Spline)
 	FVector GetLocationAtSplinePoint(int32 PointIndex, ESplineCoordinateSpace::Type CoordinateSpace) const;
diff --git a/Engine/Source/Runtime/Engine/Private/Components/SplineComponent.cpp b/Engine/Source/Runtime/Engine/Private/Components/SplineComponent.cpp
--- a/Engine/Source/Runtime/Engine/Private/Components/SplineComponent.cpp
+++ b/Engine/Source/Runtime/Engine/Private/Components/SplineComponent.cpp
@@ -686,7 +686,7 @@ void USplineComponent::AddSplinePointAtIndex(const FVector& Position, int32 Inde
 		USplineMetadata* Metadata = GetSplinePointsMetadata();
 		if (Metadata)
 		{
-			Metadata->InsertPoint(InKey, Index);
+			Metadata->InsertPoint(InKey, Index, bClosedLoop);
 		}
 		
 		NumPoints++;
@@ -903,6 +903,12 @@ int32 USplineComponent::GetNumberOfSplinePoints() const
 	return NumPoints;
 }
 
+int32 USplineComponent::GetNumberOfSplineSegments() const
+{
+	const int32 NumPoints = SplineCurves.Position.Points.Num();
+	return (bClosedLoop ? NumPoints : NumPoints - 1);
+}
+
 
 FVector USplineComponent::GetLocationAtSplinePoint(int32 PointIndex, ESplineCoordinateSpace::Type CoordinateSpace) const
 {
