Changes to facilitate derived spline component visualizers.

diff --git a/Engine/Source/Editor/ComponentVisualizers/Private/SplineComponentVisualizer.cpp b/Engine/Source/Editor/ComponentVisualizers/Private/SplineComponentVisualizer.cpp
--- a/Engine/Source/Editor/ComponentVisualizers/Private/SplineComponentVisualizer.cpp
+++ b/Engine/Source/Editor/ComponentVisualizers/Private/SplineComponentVisualizer.cpp
@@ -2446,46 +2446,53 @@ bool FSplineComponentVisualizer::CanSelectAllSplinePoints() const
 TSharedPtr<SWidget> FSplineComponentVisualizer::GenerateContextMenu() const
 {
 	FMenuBuilder MenuBuilder(true, SplineComponentVisualizerActions);
+	
+	GenerateContextMenuSections(MenuBuilder);
+
+	TSharedPtr<SWidget> MenuWidget = MenuBuilder.MakeWidget();
+	return MenuWidget;
+}
+
+void FSplineComponentVisualizer::GenerateContextMenuSections(FMenuBuilder& InMenuBuilder) const
+{
+	InMenuBuilder.BeginSection("SplinePointEdit", LOCTEXT("SplinePoint", "Spline Point"));
 	{
-		MenuBuilder.BeginSection("SplinePointEdit", LOCTEXT("SplinePoint", "Spline Point"));
+		if (SelectedSegmentIndex != INDEX_NONE)
 		{
-			if (SelectedSegmentIndex != INDEX_NONE)
-			{
-				MenuBuilder.AddMenuEntry(FSplineComponentVisualizerCommands::Get().AddKey);
-			}
-			else if (LastKeyIndexSelected != INDEX_NONE)
+			InMenuBuilder.AddMenuEntry(FSplineComponentVisualizerCommands::Get().AddKey);
+		}
+		else if (LastKeyIndexSelected != INDEX_NONE)
+		{
+			InMenuBuilder.AddMenuEntry(FSplineComponentVisualizerCommands::Get().DeleteKey);
+			InMenuBuilder.AddMenuEntry(FSplineComponentVisualizerCommands::Get().DuplicateKey);
+			InMenuBuilder.AddMenuEntry(FSplineComponentVisualizerCommands::Get().SelectAll);
+
+			InMenuBuilder.AddSubMenu(
+				LOCTEXT("SplinePointType", "Spline Point Type"),
+				LOCTEXT("KeyTypeTooltip", "Define the type of the spline point."),
+				FNewMenuDelegate::CreateSP(this, &FSplineComponentVisualizer::GenerateSplinePointTypeSubMenu));
+
+			// Only add the Automatic Tangents submenu if any of the keys is a curve type
+			USplineComponent* SplineComp = GetEditedSplineComponent();
+			if (SplineComp != nullptr)
 			{
-				MenuBuilder.AddMenuEntry(FSplineComponentVisualizerCommands::Get().DeleteKey);
-				MenuBuilder.AddMenuEntry(FSplineComponentVisualizerCommands::Get().DuplicateKey);
-				MenuBuilder.AddMenuEntry(FSplineComponentVisualizerCommands::Get().SelectAll);
-
-				MenuBuilder.AddSubMenu(
-					LOCTEXT("SplinePointType", "Spline Point Type"),
-					LOCTEXT("KeyTypeTooltip", "Define the type of the spline point."),
-					FNewMenuDelegate::CreateSP(this, &FSplineComponentVisualizer::GenerateSplinePointTypeSubMenu));
-
-				// Only add the Automatic Tangents submenu if any of the keys is a curve type
-				USplineComponent* SplineComp = GetEditedSplineComponent();
-				if (SplineComp != nullptr)
+				for (int32 SelectedKeyIndex : SelectedKeys)
 				{
-					for (int32 SelectedKeyIndex : SelectedKeys)
+					check(SelectedKeyIndex >= 0);
+					check(SelectedKeyIndex < SplineComp->GetNumberOfSplinePoints());
+					const auto& Point = SplineComp->SplineCurves.Position.Points[SelectedKeyIndex];
+					if (Point.IsCurveKey())
 					{
-						check(SelectedKeyIndex >= 0);
-						check(SelectedKeyIndex < SplineComp->GetNumberOfSplinePoints());
-						const auto& Point = SplineComp->SplineCurves.Position.Points[SelectedKeyIndex];
-						if (Point.IsCurveKey())
-						{
-							MenuBuilder.AddSubMenu(
-								LOCTEXT("ResetToAutomaticTangent", "Reset to Automatic Tangent"),
-								LOCTEXT("ResetToAutomaticTangentTooltip", "Reset the spline point tangent to an automatically generated value."),
-								FNewMenuDelegate::CreateSP(this, &FSplineComponentVisualizer::GenerateTangentTypeSubMenu));
-							break;
-						}
+						InMenuBuilder.AddSubMenu(
+							LOCTEXT("ResetToAutomaticTangent", "Reset to Automatic Tangent"),
+							LOCTEXT("ResetToAutomaticTangentTooltip", "Reset the spline point tangent to an automatically generated value."),
+							FNewMenuDelegate::CreateSP(this, &FSplineComponentVisualizer::GenerateTangentTypeSubMenu));
+						break;
 					}
 				}
 			}
 
-			MenuBuilder.AddMenuEntry(
+			InMenuBuilder.AddMenuEntry(
 				LOCTEXT("SplineGenerate", "Spline Generation Panel"),
 				LOCTEXT("SplineGenerateTooltip", "Opens up a spline generation panel to easily create basic shapes with splines"),
 				FSlateIcon(),
@@ -2495,42 +2502,39 @@ TSharedPtr<SWidget> FSplineComponentVisualizer::GenerateContextMenu() const
 				)
 			);
 		}
-		MenuBuilder.EndSection();
+	}
+	InMenuBuilder.EndSection();
 
-		MenuBuilder.BeginSection("Transform");
-		{
-			MenuBuilder.AddMenuEntry(FEditorViewportCommands::Get().FocusViewportToSelection);
-
-			MenuBuilder.AddSubMenu(
-				LOCTEXT("SnapAlign", "Snap/Align"),
-				LOCTEXT("KeyTypeTooltip", "Snap align options."),
-				FNewMenuDelegate::CreateSP(this, &FSplineComponentVisualizer::GenerateSnapAlignSubMenu));
-
-			/* temporarily disabled
-			MenuBuilder.AddSubMenu(
-				LOCTEXT("LockAxis", "Lock Axis"),
-				LOCTEXT("KeyTypeTooltip", "Axis to lock when adding new spline points."),
-				FNewMenuDelegate::CreateSP(this, &FSplineComponentVisualizer::GenerateLockAxisSubMenu));
-				*/
-		}
-		MenuBuilder.EndSection();
+	InMenuBuilder.BeginSection("Transform");
+	{
+		InMenuBuilder.AddMenuEntry(FEditorViewportCommands::Get().FocusViewportToSelection);
 
-		MenuBuilder.BeginSection("Spline", LOCTEXT("Spline", "Spline"));
-		{
-			MenuBuilder.AddMenuEntry(FSplineComponentVisualizerCommands::Get().ResetToDefault);
-		}
-		MenuBuilder.EndSection();
+		InMenuBuilder.AddSubMenu(
+			LOCTEXT("SnapAlign", "Snap/Align"),
+			LOCTEXT("KeyTypeTooltip", "Snap align options."),
+			FNewMenuDelegate::CreateSP(this, &FSplineComponentVisualizer::GenerateSnapAlignSubMenu));
 
-		MenuBuilder.BeginSection("Visualization", LOCTEXT("Visualization", "Visualization"));
-		{
-			MenuBuilder.AddMenuEntry(FSplineComponentVisualizerCommands::Get().VisualizeRollAndScale);
-			MenuBuilder.AddMenuEntry(FSplineComponentVisualizerCommands::Get().DiscontinuousSpline);
-		}
-		MenuBuilder.EndSection();
+		/* temporarily disabled
+		InMenuBuilder.AddSubMenu(
+			LOCTEXT("LockAxis", "Lock Axis"),
+			LOCTEXT("KeyTypeTooltip", "Axis to lock when adding new spline points."),
+			FNewMenuDelegate::CreateSP(this, &FSplineComponentVisualizer::GenerateLockAxisSubMenu));
+			*/
 	}
+	InMenuBuilder.EndSection();
 
-	TSharedPtr<SWidget> MenuWidget = MenuBuilder.MakeWidget();
-	return MenuWidget;
+	InMenuBuilder.BeginSection("Spline", LOCTEXT("Spline", "Spline"));
+	{
+		InMenuBuilder.AddMenuEntry(FSplineComponentVisualizerCommands::Get().ResetToDefault);
+	}
+	InMenuBuilder.EndSection();
+
+	InMenuBuilder.BeginSection("Visualization", LOCTEXT("Visualization", "Visualization"));
+	{
+		InMenuBuilder.AddMenuEntry(FSplineComponentVisualizerCommands::Get().VisualizeRollAndScale);
+		InMenuBuilder.AddMenuEntry(FSplineComponentVisualizerCommands::Get().DiscontinuousSpline);
+	}
+	InMenuBuilder.EndSection();
 }
 
 void FSplineComponentVisualizer::GenerateSplinePointTypeSubMenu(FMenuBuilder& MenuBuilder) const
diff --git a/Engine/Source/Editor/ComponentVisualizers/Public/SplineComponentVisualizer.h b/Engine/Source/Editor/ComponentVisualizers/Public/SplineComponentVisualizer.h
--- a/Engine/Source/Editor/ComponentVisualizers/Public/SplineComponentVisualizer.h
+++ b/Engine/Source/Editor/ComponentVisualizers/Public/SplineComponentVisualizer.h
@@ -113,6 +113,9 @@ public:
 	virtual bool IsVisualizingArchetype() const override;
 	//~ End FComponentVisualizer Interface
 
+	/** Add menu sections to the context menu */
+	virtual void GenerateContextMenuSections(FMenuBuilder& InMenuBuilder) const;
+
 	/** Get the spline component we are currently editing */
 	USplineComponent* GetEditedSplineComponent() const;
 
@@ -127,7 +130,7 @@ protected:
 	bool TransformSelectedKeys(const FVector& DeltaTranslate, const FRotator& DeltaRotate = FRotator::ZeroRotator, const FVector& DeltaScale = FVector::ZeroVector);
 
 	/** Update the key selection state of the visualizer */
-	void ChangeSelectionState(int32 Index, bool bIsCtrlHeld);
+	virtual void ChangeSelectionState(int32 Index, bool bIsCtrlHeld);
 
 	/** Alt-drag: duplicates the selected spline key */
 	virtual bool DuplicateKeyForAltDrag(const FVector& InDrag);
@@ -263,9 +266,9 @@ protected:
 	    selected spline point's X, Y, or Z value when fix is not equal to none. */
 	EAxis::Type AddKeyLockedAxis;
 
-private:
 	UProperty* SplineCurvesProperty;
 
+private:
 	TSharedPtr<SSplineGeneratorPanel> SplineGeneratorPanel;
 	static TWeakPtr<SWindow> WeakExistingWindow;
 };
