Fix spline component visualizer Focus Selected menu item. Interpolate spline component metadata when adding point to segment with right-click.

diff --git a/Engine/Source/Editor/ComponentVisualizers/Private/SplineComponentVisualizer.cpp b/Engine/Source/Editor/ComponentVisualizers/Private/SplineComponentVisualizer.cpp
--- a/Engine/Source/Editor/ComponentVisualizers/Private/SplineComponentVisualizer.cpp
+++ b/Engine/Source/Editor/ComponentVisualizers/Private/SplineComponentVisualizer.cpp
@@ -55,6 +55,7 @@ public:
 		UI_COMMAND(SetKeyToCurve, "Curve", "Set spline point to Curve type", EUserInterfaceActionType::RadioButton, FInputChord());
 		UI_COMMAND(SetKeyToLinear, "Linear", "Set spline point to Linear type", EUserInterfaceActionType::RadioButton, FInputChord());
 		UI_COMMAND(SetKeyToConstant, "Constant", "Set spline point to Constant type", EUserInterfaceActionType::RadioButton, FInputChord());
+		UI_COMMAND(FocusViewportToSelection, "Focus Selected", "Moves the camera in front of the selection", EUserInterfaceActionType::Button, FInputChord(EKeys::F));
 		UI_COMMAND(SnapToNearestSplinePoint, "Snap to Nearest Spline Point", "Snap to nearest spline point.", EUserInterfaceActionType::Button, FInputChord());
 		UI_COMMAND(AlignToNearestSplinePoint, "Align to Nearest Spline Point", "Align to nearest spline point.", EUserInterfaceActionType::Button, FInputChord());
 		UI_COMMAND(AlignPerpendicularToNearestSplinePoint, "Align Perpendicular to Nearest Spline Point", "Align perpendicular to nearest spline point.", EUserInterfaceActionType::Button, FInputChord());
@@ -98,6 +99,9 @@ public:
 	/** Set spline key to Constant type */
 	TSharedPtr<FUICommandInfo> SetKeyToConstant;
 
+	/** Focus on selection */
+	TSharedPtr<FUICommandInfo> FocusViewportToSelection;
+
 	/** Snap to nearest spline point on another spline component */
 	TSharedPtr<FUICommandInfo> SnapToNearestSplinePoint;
 
@@ -213,6 +217,11 @@ void FSplineComponentVisualizer::OnRegister()
 		FCanExecuteAction(),
 		FIsActionChecked::CreateSP(this, &FSplineComponentVisualizer::IsKeyTypeSet, CIM_Constant));
 
+	SplineComponentVisualizerActions->MapAction(
+		Commands.FocusViewportToSelection,
+		FExecuteAction::CreateStatic( &FLevelEditorActionCallbacks::ExecuteExecCommand, FString( TEXT("CAMERA ALIGN ACTIVEVIEWPORTONLY") ) )
+	);
+
 	SplineComponentVisualizerActions->MapAction(
 		Commands.SnapToNearestSplinePoint,
 		FExecuteAction::CreateSP(this, &FSplineComponentVisualizer::OnSnapToNearestSplinePoint, ESplineComponentSnapMode::Snap),
@@ -1850,7 +1859,7 @@ void FSplineComponentVisualizer::SplitSegment(const FVector& InWorldPos, int32 I
 
 	if (SplineMetadata)
 	{
-		SplineMetadata->DuplicatePoint(SegmentEndIndex);
+		SplineMetadata->InsertPoint(SegmentEndIndex, t, SplineComp->IsClosedLoop());
 	}
 
 	// Adjust input keys of subsequent points
@@ -2507,7 +2516,7 @@ void FSplineComponentVisualizer::GenerateContextMenuSections(FMenuBuilder& InMen
 
 	InMenuBuilder.BeginSection("Transform");
 	{
-		InMenuBuilder.AddMenuEntry(FEditorViewportCommands::Get().FocusViewportToSelection);
+		InMenuBuilder.AddMenuEntry(FSplineComponentVisualizerCommands::Get().FocusViewportToSelection);
 
 		InMenuBuilder.AddSubMenu(
 			LOCTEXT("SnapAlign", "Snap/Align"),
