Increase spline hit testing tolerance for spline component hit proxies. QOL fix.

diff --git a/Engine/Source/Editor/ComponentVisualizers/Private/SplineComponentVisualizer.cpp b/Engine/Source/Editor/ComponentVisualizers/Private/SplineComponentVisualizer.cpp
--- a/Engine/Source/Editor/ComponentVisualizers/Private/SplineComponentVisualizer.cpp
+++ b/Engine/Source/Editor/ComponentVisualizers/Private/SplineComponentVisualizer.cpp
@@ -351,6 +351,10 @@ void FSplineComponentVisualizer::DrawVisualization(const UActorComponent* Compon
 		const float GrabHandleSize = 10.0f;
 		const float TangentHandleSize = 8.0f;
 
+		static const float HitTestLineTolerance = 20.0f;
+
+		const float DefaultLineThickness = (bIsSplineEditable && PDI->IsHitTesting()) ? HitTestLineTolerance : 0.0f;
+
 		// Draw the tangent handles before anything else so they will not overdraw the rest of the spline
 		if (SplineComp == EditedSplineComp)
 		{
@@ -473,16 +477,16 @@ void FSplineComponentVisualizer::DrawVisualization(const UActorComponent* Compon
 						const FVector NewRightVector = SplineComp->GetRightVectorAtSplineInputKey(Key, ESplineCoordinateSpace::World);
 						const FVector NewScale = SplineComp->GetScaleAtSplineInputKey(Key) * DefaultScale;
 
-						PDI->DrawLine(OldPos, NewPos, LineColor, SDPG_Foreground);
+						PDI->DrawLine(OldPos, NewPos, LineColor, SDPG_Foreground, DefaultLineThickness);
 						if (bShouldVisualizeScale)
 						{
-							PDI->DrawLine(OldPos - OldRightVector * OldScale.Y, NewPos - NewRightVector * NewScale.Y, LineColor, SDPG_Foreground);
-							PDI->DrawLine(OldPos + OldRightVector * OldScale.Y, NewPos + NewRightVector * NewScale.Y, LineColor, SDPG_Foreground);
+							PDI->DrawLine(OldPos - OldRightVector * OldScale.Y, NewPos - NewRightVector * NewScale.Y, LineColor, SDPG_Foreground, DefaultLineThickness);
+							PDI->DrawLine(OldPos + OldRightVector * OldScale.Y, NewPos + NewRightVector * NewScale.Y, LineColor, SDPG_Foreground, DefaultLineThickness);
 
 							#if VISUALIZE_SPLINE_UPVECTORS
 							const FVector NewUpVector = SplineComp->GetUpVectorAtSplineInputKey(Key, ESplineCoordinateSpace::World);
-							PDI->DrawLine(NewPos, NewPos + NewUpVector * SplineComp->ScaleVisualizationWidth * 0.5f, LineColor, SDPG_Foreground);
-							PDI->DrawLine(NewPos, NewPos + NewRightVector * SplineComp->ScaleVisualizationWidth * 0.5f, LineColor, SDPG_Foreground);
+							PDI->DrawLine(NewPos, NewPos + NewUpVector * SplineComp->ScaleVisualizationWidth * 0.5f, LineColor, SDPG_Foreground, DefaultLineThickness);
+							PDI->DrawLine(NewPos, NewPos + NewRightVector * SplineComp->ScaleVisualizationWidth * 0.5f, LineColor, SDPG_Foreground, DefaultLineThickness);
 							#endif
 						}
 
