Spline Generation: Fix typo in edit condition for line length - bUpToNextPoint

diff --git a/Engine/Source/Editor/ComponentVisualizers/Public/SplineGeneratorPanel.h b/Engine/Source/Editor/ComponentVisualizers/Public/SplineGeneratorPanel.h
--- a/Engine/Source/Editor/ComponentVisualizers/Public/SplineGeneratorPanel.h
+++ b/Engine/Source/Editor/ComponentVisualizers/Public/SplineGeneratorPanel.h
@@ -238,7 +238,7 @@ public:
 	UPROPERTY(Transient, EditAnywhere, Category = ShapeParameters, Meta = (ClampMin = "1"))
 	int32 NumberOfPoints;
 
-	UPROPERTY(Transient, EditAnywhere, Category = ShapeParameters, Meta = (ClampMin = "0", EditCondition = "!bUpUntilNextPoint"))
+	UPROPERTY(Transient, EditAnywhere, Category = ShapeParameters, Meta = (ClampMin = "0", EditCondition = "!bUpToNextPoint"))
 	float Length;
 
 	UPROPERTY()
