Spline Generation: Window does not support maximize/minimize

diff --git a/Engine/Source/Editor/ComponentVisualizers/Private/SplineComponentVisualizer.cpp b/Engine/Source/Editor/ComponentVisualizers/Private/SplineComponentVisualizer.cpp
--- a/Engine/Source/Editor/ComponentVisualizers/Private/SplineComponentVisualizer.cpp
+++ b/Engine/Source/Editor/ComponentVisualizers/Private/SplineComponentVisualizer.cpp
@@ -2594,7 +2594,9 @@ void FSplineComponentVisualizer::CreateSplineGeneratorPanel()
 			.ScreenPosition(FSlateApplication::Get().GetCursorPos())
 			.Title(FText::FromString("Spline Generation"))
 			.SizingRule(ESizingRule::Autosized)
-			.AutoCenter(EAutoCenter::None);
+			.AutoCenter(EAutoCenter::None)
+			.SupportsMaximize(false)
+			.SupportsMinimize(false);
 
 		ExistingWindow->SetOnWindowClosed(FOnWindowClosed::CreateSP(SplineGeneratorPanel.ToSharedRef(), &SSplineGeneratorPanel::OnWindowClosed));
 
