Fix to SplineComponentVisualizer so that it handles external code changing the number of spline points from underneath it.

diff --git a/Engine/Source/Editor/ComponentVisualizers/Private/SplineComponentVisualizer.cpp b/Engine/Source/Editor/ComponentVisualizers/Private/SplineComponentVisualizer.cpp
--- a/Engine/Source/Editor/ComponentVisualizers/Private/SplineComponentVisualizer.cpp
+++ b/Engine/Source/Editor/ComponentVisualizers/Private/SplineComponentVisualizer.cpp
@@ -2,6 +2,7 @@
 
 #include "SplineComponentVisualizer.h"
 #include "CoreMinimal.h"
+#include "Algo/AnyOf.h"
 #include "Framework/Application/SlateApplication.h"
 #include "Framework/Commands/InputChord.h"
 #include "Framework/Commands/Commands.h"
@@ -710,15 +711,17 @@ bool FSplineComponentVisualizer::GetWidgetLocation(const FEditorViewportClient*
 		{
 			// Otherwise use the last key index set
 			check(LastKeyIndexSelected >= 0);
-			check(LastKeyIndexSelected < Position.Points.Num());
-			check(SelectedKeys.Contains(LastKeyIndexSelected));
-			const FInterpCurvePointVector& Point = Position.Points[LastKeyIndexSelected];
-			OutLocation = SplineComp->GetComponentTransform().TransformPosition(Point.OutVal);
-			if (!DuplicateDelayAccumulatedDrag.IsZero())
+			if (LastKeyIndexSelected < Position.Points.Num())
 			{
-				OutLocation += DuplicateDelayAccumulatedDrag;
+				check(SelectedKeys.Contains(LastKeyIndexSelected));
+				const FInterpCurvePointVector& Point = Position.Points[LastKeyIndexSelected];
+				OutLocation = SplineComp->GetComponentTransform().TransformPosition(Point.OutVal);
+				if (!DuplicateDelayAccumulatedDrag.IsZero())
+				{
+					OutLocation += DuplicateDelayAccumulatedDrag;
+				}
+				return true;
 			}
-			return true;
 		}
 	}
 
@@ -749,11 +752,26 @@ bool FSplineComponentVisualizer::IsVisualizingArchetype() const
 }
 
 
+bool FSplineComponentVisualizer::IsAnySelectedKeyIndexOutOfRange(const USplineComponent* Comp) const
+{
+	const int32 NumPoints = Comp->GetSplinePointsPosition().Points.Num();
+
+	return Algo::AnyOf(SelectedKeys, [NumPoints](int32 Index) { return Index >= NumPoints; });
+}
+
+
 bool FSplineComponentVisualizer::HandleInputDelta(FEditorViewportClient* ViewportClient, FViewport* Viewport, FVector& DeltaTranslate, FRotator& DeltaRotate, FVector& DeltaScale)
 {
 	USplineComponent* SplineComp = GetEditedSplineComponent();
 	if (SplineComp != nullptr)
 	{
+		if (IsAnySelectedKeyIndexOutOfRange(SplineComp))
+		{
+			// Something external has changed the number of spline points, meaning that the cached selected keys are no longer valid
+			EndEditing();
+			return false;
+		}
+
 		if (SelectedTangentHandle != INDEX_NONE)
 		{
 			return TransformSelectedTangent(DeltaTranslate);
@@ -958,9 +976,16 @@ bool FSplineComponentVisualizer::HandleInputKey(FEditorViewportClient* ViewportC
 {
 	bool bHandled = false;
 
+	USplineComponent* SplineComp = GetEditedSplineComponent();
+	if (SplineComp != nullptr && IsAnySelectedKeyIndexOutOfRange(SplineComp))
+	{
+		// Something external has changed the number of spline points, meaning that the cached selected keys are no longer valid
+		EndEditing();
+		return false;
+	}
+
 	if (Key == EKeys::LeftMouseButton && Event == IE_Released)
 	{
-		USplineComponent* SplineComp = GetEditedSplineComponent();
 		if (SplineComp != nullptr)
 		{
 			// Recache widget rotation
diff --git a/Engine/Source/Editor/ComponentVisualizers/Public/SplineComponentVisualizer.h b/Engine/Source/Editor/ComponentVisualizers/Public/SplineComponentVisualizer.h
--- a/Engine/Source/Editor/ComponentVisualizers/Public/SplineComponentVisualizer.h
+++ b/Engine/Source/Editor/ComponentVisualizers/Public/SplineComponentVisualizer.h
@@ -123,6 +123,9 @@ public:
 
 protected:
 
+	/** Determine if any selected key index is out of range (perhaps because something external has modified the spline) */
+	bool IsAnySelectedKeyIndexOutOfRange(const USplineComponent* Comp) const;
+
 	/** Transforms selected tangent by given translation */
 	bool TransformSelectedTangent(const FVector& DeltaTranslate);
 
