Fix selected spline tangent color.

diff --git a/Engine/Source/Editor/ComponentVisualizers/Private/SplineComponentVisualizer.cpp b/Engine/Source/Editor/ComponentVisualizers/Private/SplineComponentVisualizer.cpp
--- a/Engine/Source/Editor/ComponentVisualizers/Private/SplineComponentVisualizer.cpp
+++ b/Engine/Source/Editor/ComponentVisualizers/Private/SplineComponentVisualizer.cpp
@@ -369,20 +369,20 @@ void FSplineComponentVisualizer::DrawVisualization(const UActorComponent* Compon
 
 					PDI->SetHitProxy(NULL);
 
-					PDI->DrawLine(Location, Location + LeaveTangent, NormalColor, SDPG_Foreground);
-					PDI->DrawLine(Location, Location - ArriveTangent, NormalColor, SDPG_Foreground);
+					PDI->DrawLine(Location, Location + LeaveTangent, SelectedColor, SDPG_Foreground);
+					PDI->DrawLine(Location, Location - ArriveTangent, SelectedColor, SDPG_Foreground);
 
 					if (bIsSplineEditable)
 					{
 						PDI->SetHitProxy(new HSplineTangentHandleProxy(Component, SelectedKey, false));
 					}
-					PDI->DrawPoint(Location + LeaveTangent, NormalColor, TangentHandleSize, SDPG_Foreground);
+					PDI->DrawPoint(Location + LeaveTangent, SelectedColor, TangentHandleSize, SDPG_Foreground);
 
 					if (bIsSplineEditable)
 					{
 						PDI->SetHitProxy(new HSplineTangentHandleProxy(Component, SelectedKey, true));
 					}
-					PDI->DrawPoint(Location - ArriveTangent, NormalColor, TangentHandleSize, SDPG_Foreground);
+					PDI->DrawPoint(Location - ArriveTangent, SelectedColor, TangentHandleSize, SDPG_Foreground);
 
 					PDI->SetHitProxy(NULL);
 				}
