Check selected spline keys before trying to render them

diff --git a/Engine/Source/Editor/ComponentVisualizers/Private/SplineComponentVisualizer.cpp b/Engine/Source/Editor/ComponentVisualizers/Private/SplineComponentVisualizer.cpp
--- a/Engine/Source/Editor/ComponentVisualizers/Private/SplineComponentVisualizer.cpp
+++ b/Engine/Source/Editor/ComponentVisualizers/Private/SplineComponentVisualizer.cpp
@@ -355,36 +355,49 @@ void FSplineComponentVisualizer::DrawVisualization(const UActorComponent* Compon
 		// Draw the tangent handles before anything else so they will not overdraw the rest of the spline
 		if (SplineComp == EditedSplineComp)
 		{
-			for (int32 SelectedKey : SelectedKeys)
+			if (SplineComp->GetNumberOfSplinePoints() == 0 && SelectedKeys.Num() > 0)
 			{
-				check(SelectedKey >= 0);
-				check(SelectedKey < SplineComp->GetNumberOfSplinePoints());
-
-				if (SplineInfo.Points[SelectedKey].IsCurveKey())
+				ChangeSelectionState(INDEX_NONE, false);
+			}
+			else
+			{
+				const TSet<int32> SelectedKeysCopy = SelectedKeys;
+				for (int32 SelectedKey : SelectedKeysCopy)
 				{
-					const FVector Location = SplineComp->GetLocationAtSplinePoint(SelectedKey, ESplineCoordinateSpace::World);
-					const FVector LeaveTangent = SplineComp->GetLeaveTangentAtSplinePoint(SelectedKey, ESplineCoordinateSpace::World);
-					const FVector ArriveTangent = SplineComp->bAllowDiscontinuousSpline ?
-						SplineComp->GetArriveTangentAtSplinePoint(SelectedKey, ESplineCoordinateSpace::World) : LeaveTangent;
-
-					PDI->SetHitProxy(NULL);
-
-					PDI->DrawLine(Location, Location + LeaveTangent, SelectedColor, SDPG_Foreground);
-					PDI->DrawLine(Location, Location - ArriveTangent, SelectedColor, SDPG_Foreground);
-
-					if (bIsSplineEditable)
+					check(SelectedKey >= 0);
+					if (SelectedKey >= SplineComp->GetNumberOfSplinePoints())
 					{
-						PDI->SetHitProxy(new HSplineTangentHandleProxy(Component, SelectedKey, false));
+						// Catch any keys that might not exist anymore due to the underlying component changing.
+						ChangeSelectionState(SelectedKey, true);
+						continue;
 					}
-					PDI->DrawPoint(Location + LeaveTangent, SelectedColor, TangentHandleSize, SDPG_Foreground);
 
-					if (bIsSplineEditable)
+					if (SplineInfo.Points[SelectedKey].IsCurveKey())
 					{
-						PDI->SetHitProxy(new HSplineTangentHandleProxy(Component, SelectedKey, true));
-					}
-					PDI->DrawPoint(Location - ArriveTangent, SelectedColor, TangentHandleSize, SDPG_Foreground);
+						const FVector Location = SplineComp->GetLocationAtSplinePoint(SelectedKey, ESplineCoordinateSpace::World);
+						const FVector LeaveTangent = SplineComp->GetLeaveTangentAtSplinePoint(SelectedKey, ESplineCoordinateSpace::World);
+						const FVector ArriveTangent = SplineComp->bAllowDiscontinuousSpline ?
+							SplineComp->GetArriveTangentAtSplinePoint(SelectedKey, ESplineCoordinateSpace::World) : LeaveTangent;
 
-					PDI->SetHitProxy(NULL);
+						PDI->SetHitProxy(NULL);
+
+						PDI->DrawLine(Location, Location + LeaveTangent, SelectedColor, SDPG_Foreground);
+						PDI->DrawLine(Location, Location - ArriveTangent, SelectedColor, SDPG_Foreground);
+
+						if (bIsSplineEditable)
+						{
+							PDI->SetHitProxy(new HSplineTangentHandleProxy(Component, SelectedKey, false));
+						}
+						PDI->DrawPoint(Location + LeaveTangent, SelectedColor, TangentHandleSize, SDPG_Foreground);
+
+						if (bIsSplineEditable)
+						{
+							PDI->SetHitProxy(new HSplineTangentHandleProxy(Component, SelectedKey, true));
+						}
+						PDI->DrawPoint(Location - ArriveTangent, SelectedColor, TangentHandleSize, SDPG_Foreground);
+
+						PDI->SetHitProxy(NULL);
+					}
 				}
 			}
 		}
