Adjust barn door visualization to match how lighting is computed from RectLights

The previous code was drawing an octogonal shape which did not match the rectangular falloff in the lighting

diff --git a/Engine/Source/Editor/ComponentVisualizers/Private/RectLightComponentVisualizer.cpp b/Engine/Source/Editor/ComponentVisualizers/Private/RectLightComponentVisualizer.cpp
--- a/Engine/Source/Editor/ComponentVisualizers/Private/RectLightComponentVisualizer.cpp
+++ b/Engine/Source/Editor/ComponentVisualizers/Private/RectLightComponentVisualizer.cpp
@@ -47,8 +47,8 @@ void FRectLightComponentVisualizer::DrawVisualization( const UActorComponent* Co
 			{
 				FVector P0(0.0f,		+0.5f * RectLightComp->SourceWidth,					-0.5f * RectLightComp->SourceHeight);
 				FVector P1(0.0f,		+0.5f * RectLightComp->SourceWidth,					+0.5f * RectLightComp->SourceHeight);
-				FVector P2(BarnDepth, 	+0.5f * RectLightComp->SourceWidth + BarnExtent,	+0.5f * RectLightComp->SourceHeight);
-				FVector P3(BarnDepth, 	+0.5f * RectLightComp->SourceWidth + BarnExtent,	-0.5f * RectLightComp->SourceHeight);
+				FVector P2(BarnDepth, 	+0.5f * RectLightComp->SourceWidth + BarnExtent,	+0.5f * RectLightComp->SourceHeight + BarnExtent);
+				FVector P3(BarnDepth, 	+0.5f * RectLightComp->SourceWidth + BarnExtent,	-0.5f * RectLightComp->SourceHeight - BarnExtent);
 				Corners[0] =  P3;
 				Corners[1] =  P2;
 
@@ -58,8 +58,8 @@ void FRectLightComponentVisualizer::DrawVisualization( const UActorComponent* Co
 			{
 				FVector P0(0.0f,		-0.5f * RectLightComp->SourceWidth,	+0.5f * RectLightComp->SourceHeight);
 				FVector P1(0.0f,		+0.5f * RectLightComp->SourceWidth,	+0.5f * RectLightComp->SourceHeight);
-				FVector P2(BarnDepth,	+0.5f * RectLightComp->SourceWidth,	+0.5f * RectLightComp->SourceHeight + BarnExtent);
-				FVector P3(BarnDepth,	-0.5f * RectLightComp->SourceWidth,	+0.5f * RectLightComp->SourceHeight + BarnExtent);
+				FVector P2(BarnDepth,	+0.5f * RectLightComp->SourceWidth + BarnExtent,	+0.5f * RectLightComp->SourceHeight + BarnExtent);
+				FVector P3(BarnDepth,	-0.5f * RectLightComp->SourceWidth - BarnExtent,	+0.5f * RectLightComp->SourceHeight + BarnExtent);
 				Corners[2] =  P2;
 				Corners[3] =  P3;
 
@@ -69,8 +69,8 @@ void FRectLightComponentVisualizer::DrawVisualization( const UActorComponent* Co
 			{
 				FVector P0(0.0f,		-0.5f * RectLightComp->SourceWidth,					-0.5f * RectLightComp->SourceHeight);
 				FVector P1(0.0f,		-0.5f * RectLightComp->SourceWidth,					+0.5f * RectLightComp->SourceHeight);
-				FVector P2(BarnDepth,	-0.5f * RectLightComp->SourceWidth - BarnExtent,	+0.5f * RectLightComp->SourceHeight);
-				FVector P3(BarnDepth,	-0.5f * RectLightComp->SourceWidth - BarnExtent,	-0.5f * RectLightComp->SourceHeight);
+				FVector P2(BarnDepth,	-0.5f * RectLightComp->SourceWidth - BarnExtent,	+0.5f * RectLightComp->SourceHeight + BarnExtent);
+				FVector P3(BarnDepth,	-0.5f * RectLightComp->SourceWidth - BarnExtent,	-0.5f * RectLightComp->SourceHeight - BarnExtent);
 				Corners[4] =  P2;
 				Corners[5] =  P3;
 
@@ -80,8 +80,8 @@ void FRectLightComponentVisualizer::DrawVisualization( const UActorComponent* Co
 			{
 				FVector P0(0.0f,		-0.5f * RectLightComp->SourceWidth,	-0.5f * RectLightComp->SourceHeight);
 				FVector P1(0.0f,		+0.5f * RectLightComp->SourceWidth,	-0.5f * RectLightComp->SourceHeight);
-				FVector P2(BarnDepth,	+0.5f * RectLightComp->SourceWidth,	-0.5f * RectLightComp->SourceHeight - BarnExtent);
-				FVector P3(BarnDepth,	-0.5f * RectLightComp->SourceWidth,	-0.5f * RectLightComp->SourceHeight - BarnExtent);
+				FVector P2(BarnDepth,	+0.5f * RectLightComp->SourceWidth + BarnExtent,	-0.5f * RectLightComp->SourceHeight - BarnExtent);
+				FVector P3(BarnDepth,	-0.5f * RectLightComp->SourceWidth - BarnExtent,	-0.5f * RectLightComp->SourceHeight - BarnExtent);
 				Corners[6] =  P3;
 				Corners[7] =  P2;
 
@@ -89,17 +89,6 @@ void FRectLightComponentVisualizer::DrawVisualization( const UActorComponent* Co
 			
 			}
 
-			// Draw occluder between barn doors
-			for (int32 CIndex=0;CIndex<8;++CIndex)
-			{
-				const FVector& C0 =  Corners[(CIndex + 1) % 8];
-				const FVector& C1 =  Corners[(CIndex + 2) % 8];
-
-				const FVector TC0 = LightTM.TransformPosition(C0);
-				const FVector TC1 = LightTM.TransformPosition(C1);
-				PDI->DrawLine(TC0, TC1, ElementColor, SDPG_World);
-			}
-
 			DrawWireBox(PDI, LightTM.ToMatrixNoScale(), Box, ElementColor, SDPG_World);
 
 		}
