Add colored selection gradient effect to spline component visualizer, colorize selected spline tangent handles

diff --git a/Engine/Source/Editor/ComponentVisualizers/Private/SplineComponentVisualizer.cpp b/Engine/Source/Editor/ComponentVisualizers/Private/SplineComponentVisualizer.cpp
--- a/Engine/Source/Editor/ComponentVisualizers/Private/SplineComponentVisualizer.cpp
+++ b/Engine/Source/Editor/ComponentVisualizers/Private/SplineComponentVisualizer.cpp
@@ -23,6 +23,7 @@
 #include "EngineUtils.h"
 #include "CanvasItem.h"
 #include "CanvasTypes.h"
+#include "Math/UnrealMathUtility.h"
 
 IMPLEMENT_HIT_PROXY(HSplineVisProxy, HComponentVisProxy);
 IMPLEMENT_HIT_PROXY(HSplineKeyProxy, HSplineVisProxy);
@@ -528,20 +529,28 @@ void FSplineComponentVisualizer::DrawVisualization(const UActorComponent* Compon
 
 						PDI->SetHitProxy(NULL);
 
-						PDI->DrawLine(Location, Location + LeaveTangent, TangentColor, SDPG_Foreground);
-						PDI->DrawLine(Location, Location - ArriveTangent, TangentColor, SDPG_Foreground);
+						// determine tangent coloration
+						const bool bTangentSelected = (SelectedKey == SelectionState->GetSelectedTangentHandle());
+						const ESelectedTangentHandle SelectedTangentHandleType = SelectionState->GetSelectedTangentHandleType();
+						const bool bArriveSelected = bTangentSelected && (SelectedTangentHandleType == ESelectedTangentHandle::Arrive);
+						const bool bLeaveSelected = bTangentSelected && (SelectedTangentHandleType == ESelectedTangentHandle::Leave);
+						FColor ArriveColor = bArriveSelected ? SelectedColor : TangentColor;
+						FColor LeaveColor = bLeaveSelected ? SelectedColor : TangentColor;
+						
+						PDI->DrawLine(Location, Location - ArriveTangent, ArriveColor, SDPG_Foreground);
+						PDI->DrawLine(Location, Location + LeaveTangent, LeaveColor, SDPG_Foreground);
 
 						if (bIsSplineEditable)
 						{
 							PDI->SetHitProxy(new HSplineTangentHandleProxy(Component, SelectedKey, false));
 						}
-						PDI->DrawPoint(Location + LeaveTangent, TangentColor, TangentHandleSize, SDPG_Foreground);
+						PDI->DrawPoint(Location + LeaveTangent, LeaveColor, TangentHandleSize, SDPG_Foreground);
 
 						if (bIsSplineEditable)
 						{
 							PDI->SetHitProxy(new HSplineTangentHandleProxy(Component, SelectedKey, true));
 						}
-						PDI->DrawPoint(Location - ArriveTangent, TangentColor, TangentHandleSize, SDPG_Foreground);
+						PDI->DrawPoint(Location - ArriveTangent, ArriveColor, TangentHandleSize, SDPG_Foreground);
 
 						PDI->SetHitProxy(NULL);
 					}
@@ -621,23 +630,57 @@ void FSplineComponentVisualizer::DrawVisualization(const UActorComponent* Compon
 				}
 				else
 				{
+					// Determine the colors to use
+					const bool bIsEdited = (SplineComp == EditedSplineComp);
+					const bool bKeyIdxLooped = (SplineInfo.bIsLooped && KeyIdx == NumPoints);
+					const int32 BeginIdx = bKeyIdxLooped ? 0 : KeyIdx;
+					const int32 EndIdx = KeyIdx - 1;
+					const bool bBeginSelected = SelectedKeys.Contains(BeginIdx);
+					const bool bEndSelected = SelectedKeys.Contains(EndIdx);
+					const FColor BeginColor = (bIsEdited && bBeginSelected) ? SelectedColor : NormalColor;
+					const FColor EndColor = (bIsEdited && bEndSelected) ? SelectedColor : NormalColor;
+					
 					// Find position on first keyframe.
 					FVector OldPos = OldKeyPos;
 					FVector OldRightVector = OldKeyRightVector;
 					FVector OldScale = OldKeyScale;
 
 					// Then draw a line for each substep.
-					const int32 NumSteps = 20;
+					constexpr int32 NumSteps = 20;
+					constexpr float PartialGradientProportion = 0.75f;
+					constexpr int32 PartialNumSteps = NumSteps * PartialGradientProportion;
 					const float SegmentLineThickness = GetDefault<ULevelEditorViewportSettings>()->SplineLineThicknessAdjustment;
 
 					for (int32 StepIdx = 1; StepIdx <= NumSteps; StepIdx++)
 					{
-						const float Key = (KeyIdx - 1) + (StepIdx / static_cast<float>(NumSteps));
+						const float StepRatio = StepIdx / static_cast<float>(NumSteps);
+						const float Key = EndIdx + StepRatio;
 						const FVector NewPos = SplineComp->GetLocationAtSplineInputKey(Key, ESplineCoordinateSpace::World);
 						const FVector NewRightVector = SplineComp->GetRightVectorAtSplineInputKey(Key, ESplineCoordinateSpace::World);
 						const FVector NewScale = SplineComp->GetScaleAtSplineInputKey(Key) * DefaultScale;
 
-						PDI->DrawLine(OldPos, NewPos, LineColor, SDPG_Foreground, SegmentLineThickness);
+						// creates a gradient that starts partway through the selection
+						FColor StepColor;
+						if (bBeginSelected == bEndSelected)
+						{
+							StepColor = BeginColor;
+						}
+						else if (bBeginSelected && StepIdx > (NumSteps - PartialNumSteps))
+						{
+							const float LerpRatio = (1.0f - StepRatio) / PartialGradientProportion;
+							StepColor = FMath::Lerp(BeginColor.ReinterpretAsLinear(), EndColor.ReinterpretAsLinear(), LerpRatio).ToFColor(false);
+						}
+						else if (bEndSelected && StepIdx <= PartialNumSteps)
+						{
+							const float LerpRatio = 1.0f - (StepRatio / PartialGradientProportion);
+							StepColor = FMath::Lerp(BeginColor.ReinterpretAsLinear(), EndColor.ReinterpretAsLinear(), LerpRatio).ToFColor(false);
+						}
+						else
+						{
+							StepColor = NormalColor; // unselected
+						}
+						
+						PDI->DrawLine(OldPos, NewPos, StepColor, SDPG_Foreground, SegmentLineThickness);
 						if (bShouldVisualizeScale)
 						{
 							PDI->DrawLine(OldPos - OldRightVector * OldScale.Y, NewPos - NewRightVector * NewScale.Y, LineColor, SDPG_Foreground);
