Fix for crash if TextureLightProfile Format() != TSF_RGBA16F

diff --git a/Engine/Source/Editor/ComponentVisualizers/Private/PointLightComponentVisualizer.cpp b/Engine/Source/Editor/ComponentVisualizers/Private/PointLightComponentVisualizer.cpp
--- a/Engine/Source/Editor/ComponentVisualizers/Private/PointLightComponentVisualizer.cpp
+++ b/Engine/Source/Editor/ComponentVisualizers/Private/PointLightComponentVisualizer.cpp
@@ -200,7 +200,10 @@ void FTextureLightProfileVisualizer::DrawVisualization(UTextureLightProfile* Tex
 {
 	using namespace TextureLightProfileVisualizerImpl;
 
-	UpdateIntensitiesCache(TextureLightProfile, LightTM);
+	if (!UpdateIntensitiesCache(TextureLightProfile, LightTM))
+	{
+		return;
+	}
 	
 	FPolarSampler::ForEach(
 		[ &LightTM, PDI, this ]( FPolarCoordinates PolarCoordinates ) -> void
@@ -241,7 +244,7 @@ void FTextureLightProfileVisualizer::DrawVisualization(UTextureLightProfile* Tex
 	);
 }
 
-void FTextureLightProfileVisualizer::UpdateIntensitiesCache(UTextureLightProfile* TextureLightProfile, const FTransform& LightTM)
+bool FTextureLightProfileVisualizer::UpdateIntensitiesCache(UTextureLightProfile* TextureLightProfile, const FTransform& LightTM)
 {
 	using namespace TextureLightProfileVisualizerImpl;
 
@@ -250,12 +253,12 @@ void FTextureLightProfileVisualizer::UpdateIntensitiesCache(UTextureLightProfile
 	{
 		CachedLightProfile = nullptr;
 		IntensitiesCache.Empty();
-		return; 
+		return false; 
 	}
 
 	if ( CachedLightProfile == TextureLightProfile )
 	{
-		return;
+		return true;
 	}
 
 	CachedLightProfile = TextureLightProfile;
@@ -289,5 +292,6 @@ void FTextureLightProfileVisualizer::UpdateIntensitiesCache(UTextureLightProfile
 			IntensitiesCache[ SampleIndex ] = LightProfileIntensity;
 		}
 	);
+	return true;
 }
 
diff --git a/Engine/Source/Editor/ComponentVisualizers/Public/PointLightComponentVisualizer.h b/Engine/Source/Editor/ComponentVisualizers/Public/PointLightComponentVisualizer.h
--- a/Engine/Source/Editor/ComponentVisualizers/Public/PointLightComponentVisualizer.h
+++ b/Engine/Source/Editor/ComponentVisualizers/Public/PointLightComponentVisualizer.h
@@ -17,7 +17,7 @@ public:
 	void DrawVisualization(UTextureLightProfile* TextureLightProfile, const FTransform& LightTM, const FSceneView* View, FPrimitiveDrawInterface* PDI);
 
 private:
-	void UpdateIntensitiesCache(UTextureLightProfile* TextureLightProfile, const FTransform& LightTM);
+	bool UpdateIntensitiesCache(UTextureLightProfile* TextureLightProfile, const FTransform& LightTM);
 
 	const UTextureLightProfile* CachedLightProfile;
 	TArray< float > IntensitiesCache;
