[Bug Fix] Selecting other components after selecting a Spline Component's control point keeps the transform gizmo and functionality at the Spline Control point

diff --git a/Engine/Source/Editor/ComponentVisualizers/Private/SplineComponentVisualizer.cpp b/Engine/Source/Editor/ComponentVisualizers/Private/SplineComponentVisualizer.cpp
--- a/Engine/Source/Editor/ComponentVisualizers/Private/SplineComponentVisualizer.cpp
+++ b/Engine/Source/Editor/ComponentVisualizers/Private/SplineComponentVisualizer.cpp
@@ -455,6 +455,11 @@ void FSplineComponentVisualizer::OnRegister()
 FSplineComponentVisualizer::~FSplineComponentVisualizer()
 {
 	FSplineComponentVisualizerCommands::Unregister();
+	USplineComponent* SplineComponent = GetEditedSplineComponent();
+	if (SplineComponent && DeselectedInEditorDelegateHandle.IsValid())
+	{
+		SplineComponent->OnDeselectedInEditor.Remove(DeselectedInEditorDelegateHandle);
+	}
 }
 
 void FSplineComponentVisualizer::AddReferencedObjects(FReferenceCollector& Collector)
@@ -808,6 +813,14 @@ void FSplineComponentVisualizer::ChangeSelectionState(int32 Index, bool bIsCtrlH
 	{
 		SplineGeneratorPanel->OnSelectionUpdated();
 	}
+
+	if (Index != INDEX_NONE)
+	{
+		if (!DeselectedInEditorDelegateHandle.IsValid())
+		{
+			DeselectedInEditorDelegateHandle = GetEditedSplineComponent()->OnDeselectedInEditor.AddRaw(this, &FSplineComponentVisualizer::OnDeselectedInEditor);
+		}
+	}
 }
 
 
@@ -3446,5 +3459,15 @@ void FSplineComponentVisualizer::CreateSplineGeneratorPanel()
 	ExistingWindow->SetContent(SplineGeneratorPanel.ToSharedRef());
 }
 
+void FSplineComponentVisualizer::OnDeselectedInEditor(TObjectPtr<USplineComponent> SplineComponent)
+{
+	if (DeselectedInEditorDelegateHandle.IsValid() && SplineComponent)
+	{
+		SplineComponent->OnDeselectedInEditor.Remove(DeselectedInEditorDelegateHandle);
+	}
+	DeselectedInEditorDelegateHandle.Reset();
+	EndEditing();
+}
+
 
 #undef LOCTEXT_NAMESPACE
diff --git a/Engine/Source/Editor/ComponentVisualizers/Public/SplineComponentVisualizer.h b/Engine/Source/Editor/ComponentVisualizers/Public/SplineComponentVisualizer.h
--- a/Engine/Source/Editor/ComponentVisualizers/Public/SplineComponentVisualizer.h
+++ b/Engine/Source/Editor/ComponentVisualizers/Public/SplineComponentVisualizer.h
@@ -382,6 +382,8 @@ protected:
 	void SetEditedSplineComponent(const USplineComponent* InSplineComponent);
 
 	void CreateSplineGeneratorPanel();
+	
+	void OnDeselectedInEditor(TObjectPtr<USplineComponent> SplineComponent);
 
 	// FGCObject interface
 	virtual void AddReferencedObjects(FReferenceCollector& Collector);
@@ -427,6 +429,8 @@ protected:
 
 	FProperty* SplineCurvesProperty;
 
+	FDelegateHandle DeselectedInEditorDelegateHandle;
+
 private:
 	TSharedPtr<SSplineGeneratorPanel> SplineGeneratorPanel;
 	static TWeakPtr<SWindow> WeakExistingWindow;
diff --git a/Engine/Source/Runtime/Engine/Classes/Components/PrimitiveComponent.h b/Engine/Source/Runtime/Engine/Classes/Components/PrimitiveComponent.h
--- a/Engine/Source/Runtime/Engine/Classes/Components/PrimitiveComponent.h
+++ b/Engine/Source/Runtime/Engine/Classes/Components/PrimitiveComponent.h
@@ -2283,7 +2283,7 @@ public:
 	/**
 	 * Pushes new selection state to the render thread primitive proxy
 	 */
-	void PushSelectionToProxy();
+	virtual void PushSelectionToProxy();
 
 	/**
 	 * Pushes new LevelInstance editing state to the render thread primitive proxy.
diff --git a/Engine/Source/Runtime/Engine/Classes/Components/SplineComponent.h b/Engine/Source/Runtime/Engine/Classes/Components/SplineComponent.h
--- a/Engine/Source/Runtime/Engine/Classes/Components/SplineComponent.h
+++ b/Engine/Source/Runtime/Engine/Classes/Components/SplineComponent.h
@@ -290,6 +290,10 @@ public:
 	/** Width of spline in editor for use with scale visualization */
 	UPROPERTY(EditAnywhere, Category = Editor, meta=(EditCondition="bShouldVisualizeScale"))
 	float ScaleVisualizationWidth;
+
+	/** Delegate that's called when this component is deselected in the editor */
+	DECLARE_MULTICAST_DELEGATE_OneParam(DeselectedInEditorDelegate, TObjectPtr<USplineComponent>)
+	DeselectedInEditorDelegate OnDeselectedInEditor;
 #endif
 
 	//~ Begin UObject Interface
@@ -308,8 +312,14 @@ public:
 #if UE_ENABLE_DEBUG_DRAWING
 	//~ Begin UPrimitiveComponent Interface.
 	virtual FPrimitiveSceneProxy* CreateSceneProxy() override;
+#endif
+	
+#if WITH_EDITOR
+	virtual void PushSelectionToProxy() override;
+#endif
 	//~ End UPrimitiveComponent Interface.
 
+#if UE_ENABLE_DEBUG_DRAWING
 	//~ Begin USceneComponent Interface
 	virtual FBoxSphereBounds CalcBounds(const FTransform& LocalToWorld) const override;
 	//~ End USceneComponent Interface
diff --git a/Engine/Source/Runtime/Engine/Private/Components/SplineComponent.cpp b/Engine/Source/Runtime/Engine/Private/Components/SplineComponent.cpp
--- a/Engine/Source/Runtime/Engine/Private/Components/SplineComponent.cpp
+++ b/Engine/Source/Runtime/Engine/Private/Components/SplineComponent.cpp
@@ -1738,6 +1738,21 @@ FPrimitiveSceneProxy* USplineComponent::CreateSceneProxy()
 
 	return new FSplineSceneProxy(this);
 }
+#endif
+	
+#if WITH_EDITOR
+
+void USplineComponent::PushSelectionToProxy()
+{
+	if (!IsComponentIndividuallySelected())
+	{
+		OnDeselectedInEditor.Broadcast(this);
+	}
+	Super::PushSelectionToProxy();
+}
+#endif
+
+#if UE_ENABLE_DEBUG_DRAWING
 
 void USplineComponent::Draw(FPrimitiveDrawInterface* PDI, const FSceneView* View, const FInterpCurveVector& SplineInfo, const FMatrix& LocalToWorld, const FLinearColor& LineColor, uint8 DepthPriorityGroup)
 {
