Splines: Fixed an issue where the spline visualizer wasn't taking into account the list of enabled/supported spline point types for the given component.

diff --git a/Engine/Source/Editor/ComponentVisualizers/Private/SplineComponentVisualizer.cpp b/Engine/Source/Editor/ComponentVisualizers/Private/SplineComponentVisualizer.cpp
--- a/Engine/Source/Editor/ComponentVisualizers/Private/SplineComponentVisualizer.cpp
+++ b/Engine/Source/Editor/ComponentVisualizers/Private/SplineComponentVisualizer.cpp
@@ -3375,9 +3375,21 @@ void FSplineComponentVisualizer::GenerateSelectSplinePointsSubMenu(FMenuBuilder&
 }
 void FSplineComponentVisualizer::GenerateSplinePointTypeSubMenu(FMenuBuilder& MenuBuilder) const
 {
-	MenuBuilder.AddMenuEntry(FSplineComponentVisualizerCommands::Get().SetKeyToCurve);
-	MenuBuilder.AddMenuEntry(FSplineComponentVisualizerCommands::Get().SetKeyToLinear);
-	MenuBuilder.AddMenuEntry(FSplineComponentVisualizerCommands::Get().SetKeyToConstant);
+	const USplineComponent* SplineComp = GetEditedSplineComponent();
+	check(SplineComp);
+	const TArray<ESplinePointType::Type> EnabledSplinePointTypes = SplineComp->GetEnabledSplinePointTypes();
+	if (EnabledSplinePointTypes.Contains(ESplinePointType::Curve))
+	{
+		MenuBuilder.AddMenuEntry(FSplineComponentVisualizerCommands::Get().SetKeyToCurve);
+	}
+	if (EnabledSplinePointTypes.Contains(ESplinePointType::Linear))
+	{
+		MenuBuilder.AddMenuEntry(FSplineComponentVisualizerCommands::Get().SetKeyToLinear);
+	}
+	if (EnabledSplinePointTypes.Contains(ESplinePointType::Constant))
+	{
+		MenuBuilder.AddMenuEntry(FSplineComponentVisualizerCommands::Get().SetKeyToConstant);
+	}
 }
 
 void FSplineComponentVisualizer::GenerateTangentTypeSubMenu(FMenuBuilder& MenuBuilder) const
