* Added proper support for property change type in spline component visualizer (interactive changes when manipulating spline points, non-interactive change at the end of the drag operation)
* Added project settings to deactivate landscape or water mesh updates during interactive changes of water body/island splines
* Reorganised Water Editor project settings

diff --git a/Engine/Plugins/Experimental/Water/Source/Editor/Private/WaterLandscapeBrush.cpp b/Engine/Plugins/Experimental/Water/Source/Editor/Private/WaterLandscapeBrush.cpp
--- a/Engine/Plugins/Experimental/Water/Source/Editor/Private/WaterLandscapeBrush.cpp
+++ b/Engine/Plugins/Experimental/Water/Source/Editor/Private/WaterLandscapeBrush.cpp
@@ -26,7 +26,8 @@
 #include "Components/BillboardComponent.h"
 #include "Modules/ModuleManager.h"
 #include "WaterModule.h"
-#include "LandscapeModule.h"
+#include "WaterEditorSettings.h"
+"LandscapeModule.h"
 #include "LandscapeEditorServices.h"
 
 #define LOCTEXT_NAMESPACE "WaterLandscapeBrush"
@@ -165,13 +166,17 @@ void AWaterLandscapeBrush::UpdateActors(bool bInTriggerEvents)
 	}
 }
 
-void AWaterLandscapeBrush::OnActorChanged(AActor* Actor, bool bWeightmapSettingsChanged, bool bRebuildWaterMesh)
+void AWaterLandscapeBrush::OnWaterBrushActorChanged(const IWaterBrushActorInterface::FWaterBrushActorChangedEventParams& InParams)
 {
-	bool bAffectsLandscape = IsActorAffectingLandscape(Actor);
-	IWaterBrushActorInterface* WaterBrushActor = CastChecked<IWaterBrushActorInterface>(Actor);
-	int32 ActorIndex = ActorsAffectingLandscape.IndexOfByKey(TWeakInterfacePtr<IWaterBrushActorInterface>(WaterBrushActor));
+	AActor* Actor = CastChecked<AActor>(InParams.WaterBrushActor);
+	bool bAffectsLandscape = InParams.WaterBrushActor->AffectsLandscape();
+	bool bAffectsWaterMesh = InParams.WaterBrushActor->AffectsWaterMesh();
+
+	int32 ActorIndex = ActorsAffectingLandscape.IndexOfByKey(InParams.WaterBrushActor);
 	// if the actor went from affecting landscape to non-affecting landscape (and vice versa), update the brush
 	bool bForceUpdateBrush = false;
+	bool bForceUpdateWaterMesh = false;
+	
 	if (bAffectsLandscape != (ActorIndex != INDEX_NONE))
 	{
 		if (bAffectsLandscape)
@@ -184,23 +189,28 @@ void AWaterLandscapeBrush::OnActorChanged(AActor* Actor, bool bWeightmapSettings
 		}
 
 		// Force rebuild the mesh if a water body actor has been added or removed (islands don't affect the water mesh so it's not necessary for them): 
-		bRebuildWaterMesh = WaterBrushActor->CanEverAffectWaterMesh();
+		bForceUpdateWaterMesh = InParams.WaterBrushActor->CanEverAffectWaterMesh();
 		bForceUpdateBrush = true;
 	}
 
-	if (bWeightmapSettingsChanged)
+	if (InParams.bWeightmapSettingsChanged)
 	{
 		UpdateAffectedWeightmaps();
 	}
 
 	BlueprintWaterBodyChanged(Actor);
 
-	if (bAffectsLandscape || bForceUpdateBrush)
+	const UWaterEditorSettings* WaterEditorSettings = GetDefault<UWaterEditorSettings>();
+	check(WaterEditorSettings != nullptr);
+
+	bool bAllowLandscapeUpdate = (InParams.PropertyChangedEvent.ChangeType != EPropertyChangeType::Interactive) || WaterEditorSettings->GetUpdateLandscapeDuringInteractiveChanges();
+	if (bForceUpdateBrush || (bAffectsLandscape && bAllowLandscapeUpdate))
 	{
 		RequestLandscapeUpdate();
 	}
 
-	if (bRebuildWaterMesh)
+	bool bAllowWaterMeshUpdate = (InParams.PropertyChangedEvent.ChangeType != EPropertyChangeType::Interactive) || WaterEditorSettings->GetUpdateWaterMeshDuringInteractiveChanges();
+	if (bForceUpdateWaterMesh || (bAffectsWaterMesh && bAllowWaterMeshUpdate))
 	{
 		if (UWaterSubsystem* WaterSubsystem = UWaterSubsystem::GetWaterSubsystem(GetWorld()))
 		{
@@ -209,12 +219,6 @@ void AWaterLandscapeBrush::OnActorChanged(AActor* Actor, bool bWeightmapSettings
 	}
 }
 
-void AWaterLandscapeBrush::OnWaterBrushActorChanged(const IWaterBrushActorInterface::FWaterBrushActorChangedEventParams& InParams)
-{
-	AActor* Actor = CastChecked<AActor>(InParams.WaterBrushActor);
-	OnActorChanged(Actor, /* bWeightmapSettingsChanged = */InParams.bWeightmapSettingsChanged, /* bRebuildWaterMesh = */InParams.WaterBrushActor->AffectsWaterMesh());
-}
-
 void AWaterLandscapeBrush::OnActorsAffectingLandscapeChanged()
 {
 	BlueprintWaterBodiesChanged();
@@ -283,14 +287,6 @@ void AWaterLandscapeBrush::PostInitProperties()
 		}
 #endif // WITH_EDITOR
 
-		OnActorMovedHandle = GEngine->OnActorMoved().AddLambda([this](AActor* Actor)
-		{
-			if (IsActorAffectingLandscape(Actor))
-			{
-				OnActorChanged(Actor, /* bWeightmapSettingsChanged */ false, /* bRebuildWaterMesh */ true);
-			}
-		});
-
 		IWaterBrushActorInterface::GetOnWaterBrushActorChangedEvent().AddUObject(this, &AWaterLandscapeBrush::OnWaterBrushActorChanged);
 	}
 
@@ -352,9 +348,6 @@ void AWaterLandscapeBrush::BeginDestroy()
 		GEngine->OnLevelActorDeleted().Remove(OnLevelActorDeletedHandle);
 		OnLevelActorDeletedHandle.Reset();
 
-		GEngine->OnActorMoved().Remove(OnActorMovedHandle);
-		OnActorMovedHandle.Reset();
-
 #if WITH_EDITOR
 		// Unregister callbacks
 		if (UWorld* World = GetWorld())
diff --git a/Engine/Plugins/Experimental/Water/Source/Editor/Private/WaterLandscapeBrush.h b/Engine/Plugins/Experimental/Water/Source/Editor/Private/WaterLandscapeBrush.h
--- a/Engine/Plugins/Experimental/Water/Source/Editor/Private/WaterLandscapeBrush.h
+++ b/Engine/Plugins/Experimental/Water/Source/Editor/Private/WaterLandscapeBrush.h
@@ -119,7 +119,6 @@ private:
 
 	void OnFullHeightmapRenderDone(UTextureRenderTarget2D* HeightmapRenderTarget);
 	void OnWaterBrushActorChanged(const IWaterBrushActorInterface::FWaterBrushActorChangedEventParams& InParams);
-	void OnActorChanged(AActor* Actor, bool bWeightmapSettingsChanged, bool bRebuildMesh);
 	void OnActorsAffectingLandscapeChanged();
 	void OnLevelActorAdded(AActor* InActor);
 	void OnLevelActorRemoved(AActor* InActor);
diff --git a/Engine/Plugins/Experimental/Water/Source/Editor/Private/WaterSplineComponentVisualizer.cpp b/Engine/Plugins/Experimental/Water/Source/Editor/Private/WaterSplineComponentVisualizer.cpp
--- a/Engine/Plugins/Experimental/Water/Source/Editor/Private/WaterSplineComponentVisualizer.cpp
+++ b/Engine/Plugins/Experimental/Water/Source/Editor/Private/WaterSplineComponentVisualizer.cpp
@@ -580,7 +580,8 @@ bool FWaterSplineComponentVisualizer::HandleInputDelta(FEditorViewportClient* Vi
 					WaterSplineComp->UpdateSpline();
 					WaterSplineComp->bSplineHasBeenEdited = true;
 
-					NotifyPropertyModified(WaterSplineComp, SplineCurvesProperty);
+					// Transform the spline keys using an EPropertyChangeType::Interactive change. Later on, at the end of mouse tracking, a non-interactive change will be notified via void TrackingStopped :
+					NotifyPropertyModified(WaterSplineComp, SplineCurvesProperty, EPropertyChangeType::Interactive);
 
 					GEditor->RedrawLevelEditingViewports(true);
 
diff --git a/Engine/Plugins/Experimental/Water/Source/Editor/Public/WaterEditorSettings.h b/Engine/Plugins/Experimental/Water/Source/Editor/Public/WaterEditorSettings.h
--- a/Engine/Plugins/Experimental/Water/Source/Editor/Public/WaterEditorSettings.h
+++ b/Engine/Plugins/Experimental/Water/Source/Editor/Public/WaterEditorSettings.h
@@ -222,6 +222,9 @@ public:
 	UMaterialInterface* GetDefaultRenderRiverSplineDepthsMaterial() const;
 	FSoftObjectPath GetDefaultRenderRiverSplineDepthsMaterialPath() const { return DefaultRenderRiverSplineDepthsMaterial.ToSoftObjectPath(); }
 
+	bool GetUpdateLandscapeDuringInteractiveChanges() const { return bUpdateLandscapeDuringInteractiveChanges; }
+	bool GetUpdateWaterMeshDuringInteractiveChanges() const { return bUpdateWaterMeshDuringInteractiveChanges; }
+
 public:
 	/** The texture group to use for generated textures such as the combined veloctiy and height texture */
 	UPROPERTY(EditAnywhere, config, Category = Rendering)
@@ -236,7 +239,7 @@ public:
 	float VisualizeWaterVelocityScale = 20.0f;
 
 	/** Material Parameter Collection for everything landscape-related */
-	UPROPERTY(EditAnywhere, config, Category = Brush)
+	UPROPERTY(EditAnywhere, config, Category = Brush, AdvancedDisplay)
 	TSoftObjectPtr<UMaterialParameterCollection> LandscapeMaterialParameterCollection;
 
 	/** Default values for base WaterMesh actor*/
@@ -263,40 +266,48 @@ public:
 	FWaterBodyIslandDefaults WaterBodyIslandDefaults;
 
 private:
+	/** For landscape-affecting water bodies, allows the landscape to be updated when the water body's shape is modified interactively (e.g. when dragging a spline point). Set to false if the performance when editing a water body gets too bad (the landscape will be properly updated when the dragging operation is done). */
+	UPROPERTY(EditAnywhere, config, Category = Brush)
+	bool bUpdateLandscapeDuringInteractiveChanges = false;
+
+	/** Allows the water mesh to be updated when the water body's shape is modified interactively (e.g. when dragging a spline point). Set to false if the performance when editing a water body gets too bad (the water mesh will be properly updated when the dragging operation is done). */
+	UPROPERTY(EditAnywhere, config, Category = Brush)
+	bool bUpdateWaterMeshDuringInteractiveChanges = false;
+
 	/** Class of the water brush to be used in landscape */
-	UPROPERTY(EditAnywhere, config, Category = Brush, meta = (MetaClass = "/Script/WaterEditor.WaterLandscapeBrush"))
+	UPROPERTY(EditAnywhere, config, Category = Brush, meta = (MetaClass = "/Script/WaterEditor.WaterLandscapeBrush"), AdvancedDisplay)
 	FSoftClassPath WaterManagerClassPath;
 
-	UPROPERTY(EditAnywhere, config, Category = Brush)
+	UPROPERTY(EditAnywhere, config, Category = Brush, AdvancedDisplay)
 	TSoftObjectPtr<UMaterialInterface> DefaultBrushAngleFalloffMaterial;
 
-	UPROPERTY(EditAnywhere, config, Category = Brush)
+	UPROPERTY(EditAnywhere, config, Category = Brush, AdvancedDisplay)
 	TSoftObjectPtr<UMaterialInterface> DefaultBrushIslandFalloffMaterial;
 
-	UPROPERTY(EditAnywhere, config, Category = Brush)
+	UPROPERTY(EditAnywhere, config, Category = Brush, AdvancedDisplay)
 	TSoftObjectPtr<UMaterialInterface> DefaultBrushWidthFalloffMaterial;
 
-	UPROPERTY(EditAnywhere, config, Category = Brush)
+	UPROPERTY(EditAnywhere, config, Category = Brush, AdvancedDisplay)
 	TSoftObjectPtr<UMaterialInterface> DefaultBrushWeightmapMaterial;
 
-	UPROPERTY(EditAnywhere, config, Category = Brush)
+	UPROPERTY(EditAnywhere, config, Category = Brush, AdvancedDisplay)
 	TSoftObjectPtr<UMaterialInterface> DefaultCacheDistanceFieldCacheMaterial;
 
-	UPROPERTY(EditAnywhere, config, Category = Brush)
+	UPROPERTY(EditAnywhere, config, Category = Brush, AdvancedDisplay)
 	TSoftObjectPtr<UMaterialInterface> DefaultCompositeWaterBodyTextureMaterial;
 
-	UPROPERTY(EditAnywhere, config, Category = Brush)
+	UPROPERTY(EditAnywhere, config, Category = Brush, AdvancedDisplay)
 	TSoftObjectPtr<UMaterialInterface> DefaultJumpFloodStepMaterial;
 
-	UPROPERTY(EditAnywhere, config, Category = Brush)
+	UPROPERTY(EditAnywhere, config, Category = Brush, AdvancedDisplay)
 	TSoftObjectPtr<UMaterialInterface> DefaultBlurEdgesMaterial;
 
-	UPROPERTY(EditAnywhere, config, Category = Brush)
+	UPROPERTY(EditAnywhere, config, Category = Brush, AdvancedDisplay)
 	TSoftObjectPtr<UMaterialInterface> DefaultFindEdgesMaterial;
 
-	UPROPERTY(EditAnywhere, config, Category = Brush)
+	UPROPERTY(EditAnywhere, config, Category = Brush, AdvancedDisplay)
 	TSoftObjectPtr<UMaterialInterface> DefaultDrawCanvasMaterial;
 
-	UPROPERTY(EditAnywhere, config, Category = Brush)
+	UPROPERTY(EditAnywhere, config, Category = Brush, AdvancedDisplay)
 	TSoftObjectPtr<UMaterialInterface> DefaultRenderRiverSplineDepthsMaterial;
 };
diff --git a/Engine/Plugins/Experimental/Water/Source/Runtime/Private/WaterBodyActor.cpp b/Engine/Plugins/Experimental/Water/Source/Runtime/Private/WaterBodyActor.cpp
--- a/Engine/Plugins/Experimental/Water/Source/Runtime/Private/WaterBodyActor.cpp
+++ b/Engine/Plugins/Experimental/Water/Source/Runtime/Private/WaterBodyActor.cpp
@@ -104,7 +104,10 @@ void AWaterBody::PreInitializeComponents()
 
 	check(WaterBodyComponent != nullptr);
 	// some water bodies are dynamic (e.g. Ocean) and thus need to be regenerated at runtime :
-	WaterBodyComponent->UpdateAll(true);
+	FOnWaterBodyChangedParams Params;
+	Params.bShapeOrPositionChanged = true;
+	Params.bWeightmapSettingsChanged = true;
+	WaterBodyComponent->UpdateAll(Params);
 }
 
 #if WITH_EDITOR
@@ -117,7 +120,10 @@ void AWaterBody::PostEditMove(bool bFinished)
 		WaterBodyComponent->UpdateWaterHeight();
 	}
 
-	WaterBodyComponent->OnWaterBodyChanged(bFinished);
+	FOnWaterBodyChangedParams Params;
+	Params.PropertyChangedEvent.ChangeType = bFinished ? EPropertyChangeType::ValueSet : EPropertyChangeType::Interactive;
+	Params.bShapeOrPositionChanged = true;
+	WaterBodyComponent->OnWaterBodyChanged(Params);
 }
 
 void AWaterBody::PreEditChange(FProperty* PropertyThatWillChange)
@@ -147,8 +153,10 @@ void AWaterBody::PostEditChangeProperty(FPropertyChangedEvent& PropertyChangedEv
 
 		WaterBodyComponent->RequestGPUWaveDataUpdate();
 
+		FOnWaterBodyChangedParams Params;
 		// Waves data affect the navigation : 
-		WaterBodyComponent->OnWaterBodyChanged(/* bShapeOrPositionChanged */ true, /* bWeightMapSettingsChanged */ false);
+		Params.bShapeOrPositionChanged = true;
+		WaterBodyComponent->OnWaterBodyChanged(Params);
 	}
 }
 
@@ -179,8 +187,10 @@ void AWaterBody::SetWaterWavesInternal(UWaterWavesBase* InWaterWaves)
 
 		WaterBodyComponent->RequestGPUWaveDataUpdate();
 
+		FOnWaterBodyChangedParams Params;
 		// Waves data can affect the navigation: 
-		WaterBodyComponent->OnWaterBodyChanged(/*bShapeOrPositionChanged = */true);
+		Params.bShapeOrPositionChanged = true;
+		WaterBodyComponent->OnWaterBodyChanged(Params);
 	}
 }
 
@@ -531,7 +541,10 @@ void AWaterBody::PostRegisterAllComponents()
 
 	if (IsValid(WaterBodyComponent))
 	{
-		WaterBodyComponent->UpdateAll(true);
+		FOnWaterBodyChangedParams Params;
+		Params.bShapeOrPositionChanged = true;
+		Params.bWeightmapSettingsChanged = true;
+		WaterBodyComponent->UpdateAll(Params);
 		WaterBodyComponent->UpdateMaterialInstances();
 
 #if WITH_EDITOR
diff --git a/Engine/Plugins/Experimental/Water/Source/Runtime/Private/WaterBodyComponent.cpp b/Engine/Plugins/Experimental/Water/Source/Runtime/Private/WaterBodyComponent.cpp
--- a/Engine/Plugins/Experimental/Water/Source/Runtime/Private/WaterBodyComponent.cpp
+++ b/Engine/Plugins/Experimental/Water/Source/Runtime/Private/WaterBodyComponent.cpp
@@ -419,7 +419,10 @@ void UWaterBodyComponent::PostDuplicate(bool bDuplicateForPie)
 		// After duplication due to copy-pasting, UWaterSplineMetadata might have been edited without the spline component being made aware of that (for some reason, USplineComponent::PostDuplicate isn't called)::
 		GetWaterSpline()->SynchronizeWaterProperties();
 
-		OnWaterBodyChanged(/*bShapeOrPositionChanged*/true, /*bWeightmapSettingsChanged*/true);
+		FOnWaterBodyChangedParams Params;
+		Params.bShapeOrPositionChanged = true;
+		Params.bWeightmapSettingsChanged = true;
+		OnWaterBodyChanged(Params);
 	}
 
 	RegisterOnUpdateWavesData(GetWaterWaves(), /* bRegister = */true);
@@ -974,7 +977,10 @@ void UWaterBodyComponent::PostEditUndo()
 	// Ensure we only trigger updates if this component is registered.
 	if (IsRegistered())
 	{
-		OnWaterBodyChanged(/*bShapeOrPositionChanged*/true, /*bWeightmapSettingsChanged*/true);
+		FOnWaterBodyChangedParams Params;
+		Params.bShapeOrPositionChanged = true;
+		Params.bWeightmapSettingsChanged = true;
+		OnWaterBodyChanged(Params);
 
 		// On undo, when PostEditChangeProperty is called, PropertyChangedEvent is fake so we need to register to the new object here :
 		RegisterOnUpdateWavesData(GetWaterWaves(), /*bRegister = */true);
@@ -987,17 +993,21 @@ void UWaterBodyComponent::PostEditImport()
 {
 	Super::PostEditImport();
 
-	OnWaterBodyChanged(/*bShapeOrPositionChanged*/true, /*bWeightmapSettingsChanged*/true);
+	FOnWaterBodyChangedParams Params;
+	Params.bShapeOrPositionChanged = true;
+	Params.bWeightmapSettingsChanged = true;
+	OnWaterBodyChanged(Params);
 
 	RequestGPUWaveDataUpdate();
 }
 
-void UWaterBodyComponent::OnPostEditChangeProperty(FPropertyChangedEvent& PropertyChangedEvent, bool& bShapeOrPositionChanged, bool& bWeightmapSettingsChanged)
+void UWaterBodyComponent::OnPostEditChangeProperty(FOnWaterBodyChangedParams& InOutOnWaterBodyChangedParams)
 {
+	const FPropertyChangedEvent& PropertyChangedEvent = InOutOnWaterBodyChangedParams.PropertyChangedEvent;
 	const FName PropertyName = PropertyChangedEvent.GetPropertyName();
 	if (PropertyChangedEvent.MemberProperty && PropertyChangedEvent.MemberProperty->GetFName() == GET_MEMBER_NAME_CHECKED(UWaterBodyComponent, LayerWeightmapSettings))
 	{
-		bWeightmapSettingsChanged = true;
+		InOutOnWaterBodyChangedParams.bWeightmapSettingsChanged = true;
 	}
 	else if ((PropertyName == GET_MEMBER_NAME_CHECKED(UWaterBodyComponent, WaterMaterial)) ||
 			(PropertyName == GET_MEMBER_NAME_CHECKED(UWaterBodyComponent, UnderwaterPostProcessMaterial)) ||
@@ -1011,7 +1021,8 @@ void UWaterBodyComponent::OnPostEditChangeProperty(FPropertyChangedEvent& Proper
 	}
 	else if (PropertyName == GET_MEMBER_NAME_CHECKED(UWaterBodyComponent, MaxWaveHeightOffset))
 	{
-		bShapeOrPositionChanged = true;
+		// Waves data affect the navigation :
+		InOutOnWaterBodyChangedParams.bShapeOrPositionChanged = true;
 	}
 	else if (PropertyChangedEvent.MemberProperty && PropertyChangedEvent.MemberProperty->GetFName() == FName(TEXT("RelativeScale3D")))
 	{
@@ -1022,6 +1033,7 @@ void UWaterBodyComponent::OnPostEditChangeProperty(FPropertyChangedEvent& Proper
 			Scale.Z = 1.f;
 			SetRelativeScale3D(Scale);
 		}
+		InOutOnWaterBodyChangedParams.bShapeOrPositionChanged = true;
 	}
 }
 
@@ -1066,22 +1078,23 @@ void UWaterBodyComponent::CheckForErrors()
 
 void UWaterBodyComponent::PostEditChangeProperty(FPropertyChangedEvent& PropertyChangedEvent)
 {
-	bool bShapeOrPositionChanged = false;
-	bool bWeightmapSettingsChanged = false;
-
-	OnPostEditChangeProperty(PropertyChangedEvent, bShapeOrPositionChanged, bWeightmapSettingsChanged);
+	FOnWaterBodyChangedParams Params(PropertyChangedEvent);
+	OnPostEditChangeProperty(Params);
 	
 	Super::PostEditChangeProperty(PropertyChangedEvent);
 
 	if (!IsTemplate())
 	{
-		OnWaterBodyChanged(bShapeOrPositionChanged, bWeightmapSettingsChanged);
+		OnWaterBodyChanged(Params);
 	}
 }
 
-void UWaterBodyComponent::OnSplineDataChanged()
+void UWaterBodyComponent::OnWaterSplineDataChanged(const FOnWaterSplineDataChangedParams& InParams)
 {
-	OnWaterBodyChanged(/*bShapeOrPositionChanged*/true);
+	// Transfer the FOnWaterSplineDataChangedParams parameters to FOnWaterBodyChangedParams :
+	FOnWaterBodyChangedParams Params(InParams.PropertyChangedEvent);
+	Params.bShapeOrPositionChanged = true;
+	OnWaterBodyChanged(Params);
 }
 
 void UWaterBodyComponent::RegisterOnUpdateWavesData(UWaterWavesBase* InWaterWaves, bool bRegister)
@@ -1103,8 +1116,11 @@ void UWaterBodyComponent::OnWavesDataUpdated(UWaterWavesBase* InWaterWaves, EPro
 {
 	RequestGPUWaveDataUpdate();
 
-	// Waves data affect the navigation : 
-	OnWaterBodyChanged(/*bShapeOrPositionChanged = */true);
+	FOnWaterBodyChangedParams Params;
+	// Waves data affect the navigation :
+	Params.PropertyChangedEvent.ChangeType = InChangeType;
+	Params.bShapeOrPositionChanged = true;
+	OnWaterBodyChanged(Params);
 }
 
 void UWaterBodyComponent::OnWaterSplineMetadataChanged(UWaterSplineMetadata* InWaterSplineMetadata, FPropertyChangedEvent& PropertyChangedEvent)
@@ -1138,11 +1154,11 @@ void UWaterBodyComponent::RegisterOnChangeWaterSplineData(bool bRegister)
 	{
 		if (bRegister)
 		{
-			WaterSpline->OnSplineDataChanged().AddUObject(this, &UWaterBodyComponent::OnSplineDataChanged);
+			WaterSpline->OnWaterSplineDataChanged().AddUObject(this, &UWaterBodyComponent::OnWaterSplineDataChanged);
 		}
 		else
 		{
-			WaterSpline->OnSplineDataChanged().RemoveAll(this);
+			WaterSpline->OnWaterSplineDataChanged().RemoveAll(this);
 		}
 	}
 
@@ -1246,12 +1262,14 @@ void UWaterBodyComponent::UpdateWaterBody(bool bWithExclusionVolumes)
 	}
 }
 
-void UWaterBodyComponent::UpdateAll(bool bShapeOrPositionChanged)
+void UWaterBodyComponent::UpdateAll(const FOnWaterBodyChangedParams& InParams)
 {
 	BeginUpdateWaterBody();
 
 	AWaterBody* WaterBodyOwner = GetWaterBodyActor();
 	check(WaterBodyOwner);
+
+	bool bShapeOrPositionChanged = InParams.bShapeOrPositionChanged;
 	
 	if (GIsEditor || IsBodyDynamic())
 	{
@@ -1315,6 +1333,14 @@ void UWaterBodyComponent::UpdateSplineComponent()
 }
 
 void UWaterBodyComponent::OnWaterBodyChanged(bool bShapeOrPositionChanged, bool bWeightmapSettingsChanged)
+{
+	FOnWaterBodyChangedParams Params;
+	Params.bShapeOrPositionChanged = bShapeOrPositionChanged;
+	Params.bWeightmapSettingsChanged = bWeightmapSettingsChanged;
+	OnWaterBodyChanged(Params);
+}
+
+void UWaterBodyComponent::OnWaterBodyChanged(const FOnWaterBodyChangedParams& InParams)
 {
 	TRACE_CPUPROFILER_EVENT_SCOPE(UWaterBodyComponent::OnWaterBodyChanged)
 	// It's possible to get called without a water spline after the Redo of a water body deletion (i.e. the water body actor gets deleted again, hence its SplineComp is restored to nullptr)
@@ -1322,10 +1348,10 @@ void UWaterBodyComponent::OnWaterBodyChanged(bool bShapeOrPositionChanged, bool
 	//  of the editor transaction) :
 	if (GetWaterSpline())
 	{
-		UpdateAll(bShapeOrPositionChanged);
+		UpdateAll(InParams);
 
 		// Some of the spline parameters need to be transferred to the underwater post process MID, if any : 
-		if (bShapeOrPositionChanged)
+		if (InParams.bShapeOrPositionChanged)
 		{
 			SetDynamicParametersOnUnderwaterPostProcessMID(UnderwaterPostProcessMID);
 		}
@@ -1333,9 +1359,10 @@ void UWaterBodyComponent::OnWaterBodyChanged(bool bShapeOrPositionChanged, bool
 
 #if WITH_EDITOR
 	AWaterBody* const WaterBodyActor = GetWaterBodyActor();
-	IWaterBrushActorInterface::FWaterBrushActorChangedEventParams Params(WaterBodyActor);
-	Params.bShapeOrPositionChanged = bShapeOrPositionChanged;
-	Params.bWeightmapSettingsChanged = bWeightmapSettingsChanged;
+	// Transfer the FOnWaterBodyChangedParams parameters to FWaterBrushActorChangedEventParams :
+	IWaterBrushActorInterface::FWaterBrushActorChangedEventParams Params(WaterBodyActor, InParams.PropertyChangedEvent);
+	Params.bShapeOrPositionChanged = InParams.bShapeOrPositionChanged;
+	Params.bWeightmapSettingsChanged = InParams.bWeightmapSettingsChanged;
 	WaterBodyActor->BroadcastWaterBrushActorChangedEvent(Params);
 #endif
 }
diff --git a/Engine/Plugins/Experimental/Water/Source/Runtime/Private/WaterBodyIslandActor.cpp b/Engine/Plugins/Experimental/Water/Source/Runtime/Private/WaterBodyIslandActor.cpp
--- a/Engine/Plugins/Experimental/Water/Source/Runtime/Private/WaterBodyIslandActor.cpp
+++ b/Engine/Plugins/Experimental/Water/Source/Runtime/Private/WaterBodyIslandActor.cpp
@@ -35,7 +35,7 @@ AWaterBodyIsland::AWaterBodyIsland(const FObjectInitializer& ObjectInitializer)
 #if WITH_EDITOR
 	if (!HasAnyFlags(RF_ClassDefaultObject))
 	{
-		SplineComp->OnSplineDataChanged().AddUObject(this, &AWaterBodyIsland::OnSplineDataChanged);
+		SplineComp->OnWaterSplineDataChanged().AddUObject(this, &AWaterBodyIsland::OnWaterSplineDataChanged);
 	}
 
 	ActorIcon = FWaterIconHelper::EnsureSpriteComponentCreated(this, TEXT("/Water/Icons/WaterBodyIslandSprite"));
@@ -231,30 +231,42 @@ void AWaterBodyIsland::PostEditMove(bool bFinished)
 {
 	Super::PostEditMove(bFinished);
 
-	UpdateAll();
+	FOnWaterBodyIslandChangedParams Params;
+	Params.PropertyChangedEvent.ChangeType = bFinished ? EPropertyChangeType::ValueSet : EPropertyChangeType::Interactive;
+	Params.bShapeOrPositionChanged = true;
+	UpdateAll(Params);
 }
 
 void AWaterBodyIsland::PostEditUndo()
 {
 	Super::PostEditUndo();
 
-	UpdateAll();
+	FOnWaterBodyIslandChangedParams Params;
+	Params.bShapeOrPositionChanged = true;
+	Params.bWeightmapSettingsChanged = true;
+	UpdateAll(Params);
 }
 
 void AWaterBodyIsland::PostEditImport()
 {
 	Super::PostEditImport();
 
-	UpdateAll();
+	FOnWaterBodyIslandChangedParams Params;
+	Params.bShapeOrPositionChanged = true;
+	Params.bWeightmapSettingsChanged = true;
+	UpdateAll(Params);
 }
 
-void AWaterBodyIsland::UpdateAll()
+void AWaterBodyIsland::UpdateAll(const FOnWaterBodyIslandChangedParams& InParams)
 {
 	UpdateHeight();
 
-	UpdateOverlappingWaterBodyComponents();
+	if (InParams.bShapeOrPositionChanged)
+	{
+		UpdateOverlappingWaterBodyComponents();
+	}
 
-	OnWaterBodyIslandChanged(/*bShapeOrPositionChanged*/true, /*bWeightmapSettingsChanged*/true);
+	OnWaterBodyIslandChanged(InParams);
 
 	UpdateActorIcon();
 }
@@ -281,30 +293,34 @@ void AWaterBodyIsland::PostEditChangeProperty(FPropertyChangedEvent& PropertyCha
 {
 	Super::PostEditChangeProperty(PropertyChangedEvent);
 
-	bool bWeightmapSettingsChanged = false;
+	FOnWaterBodyIslandChangedParams Params(PropertyChangedEvent);
 	if (PropertyChangedEvent.GetPropertyName() == GET_MEMBER_NAME_CHECKED(UWaterBodyComponent, LayerWeightmapSettings))
 	{
-		bWeightmapSettingsChanged = true;
+		Params.bWeightmapSettingsChanged = true;
 	}
 
-	OnWaterBodyIslandChanged(/*bShapeOrPositionChanged*/false, bWeightmapSettingsChanged);
+	OnWaterBodyIslandChanged(Params);
 
 	UpdateActorIcon();
 }
 
-void AWaterBodyIsland::OnSplineDataChanged()
+void AWaterBodyIsland::OnWaterSplineDataChanged(const FOnWaterSplineDataChangedParams& InParams)
 {
 	UpdateOverlappingWaterBodyComponents();
 
-	OnWaterBodyIslandChanged(/* bShapeOrPositionChanged = */true, /* bWeightmapSettingsChanged = */false);
+	// Transfer the FOnWaterSplineDataChangedParams parameters to FOnWaterBodyIslandChangedParams :
+	FOnWaterBodyIslandChangedParams Params(InParams.PropertyChangedEvent);
+	Params.bShapeOrPositionChanged = true;
+	OnWaterBodyIslandChanged(Params);
 }
 
-void AWaterBodyIsland::OnWaterBodyIslandChanged(bool bShapeOrPositionChanged, bool bWeightmapSettingsChanged)
+void AWaterBodyIsland::OnWaterBodyIslandChanged(const FOnWaterBodyIslandChangedParams& InParams)
 {
 #if WITH_EDITOR
-	FWaterBrushActorChangedEventParams Params(this);
-	Params.bShapeOrPositionChanged = bShapeOrPositionChanged;
-	Params.bWeightmapSettingsChanged = bWeightmapSettingsChanged;
+	// Transfer the FOnWaterBodyIslandChangedParams parameters to FWaterBrushActorChangedEventParams :
+	FWaterBrushActorChangedEventParams Params(this, InParams.PropertyChangedEvent);
+	Params.bShapeOrPositionChanged = InParams.bShapeOrPositionChanged;
+	Params.bWeightmapSettingsChanged = InParams.bWeightmapSettingsChanged;
 	BroadcastWaterBrushActorChangedEvent(Params);
 #endif
 }
diff --git a/Engine/Plugins/Experimental/Water/Source/Runtime/Private/WaterBodyOceanComponent.cpp b/Engine/Plugins/Experimental/Water/Source/Runtime/Private/WaterBodyOceanComponent.cpp
--- a/Engine/Plugins/Experimental/Water/Source/Runtime/Private/WaterBodyOceanComponent.cpp
+++ b/Engine/Plugins/Experimental/Water/Source/Runtime/Private/WaterBodyOceanComponent.cpp
@@ -61,7 +61,9 @@ void UWaterBodyOceanComponent::SetHeightOffset(float InHeightOffset)
 		HeightOffset = ClampedHeightOffset;
 
 		// the physics volume needs to be adjusted : 
-		OnWaterBodyChanged(true);
+		FOnWaterBodyChangedParams Params;
+		Params.bShapeOrPositionChanged = true;
+		OnWaterBodyChanged(Params);
 	}
 }
 
@@ -87,18 +89,19 @@ void UWaterBodyOceanComponent::BeginUpdateWaterBody()
 }
 
 #if WITH_EDITOR
-void UWaterBodyOceanComponent::OnPostEditChangeProperty(FPropertyChangedEvent& PropertyChangedEvent, bool& bShapeOrPositionChanged, bool& bWeightmapSettingsChanged)
+void UWaterBodyOceanComponent::OnPostEditChangeProperty(FOnWaterBodyChangedParams& InOutOnWaterBodyChangedParams)
 {
-	Super::OnPostEditChangeProperty(PropertyChangedEvent, bShapeOrPositionChanged, bWeightmapSettingsChanged);
+	Super::OnPostEditChangeProperty(InOutOnWaterBodyChangedParams);
 
+	const FPropertyChangedEvent& PropertyChangedEvent = InOutOnWaterBodyChangedParams.PropertyChangedEvent;
 	if (PropertyChangedEvent.GetPropertyName() == GET_MEMBER_NAME_CHECKED(UWaterBodyOceanComponent, CollisionExtents))
 	{
 		// Affects the physics shape
-		bShapeOrPositionChanged = true;
+		InOutOnWaterBodyChangedParams.bShapeOrPositionChanged = true;
 	}
 	else if (PropertyChangedEvent.GetPropertyName() == GET_MEMBER_NAME_CHECKED(UWaterBodyOceanComponent, VisualExtents))
 	{
-		bShapeOrPositionChanged = true;
+		InOutOnWaterBodyChangedParams.bShapeOrPositionChanged = true;
 	}
 }
 const TCHAR* UWaterBodyOceanComponent::GetWaterSpriteTextureName() const
diff --git a/Engine/Plugins/Experimental/Water/Source/Runtime/Private/WaterBodyRiverComponent.cpp b/Engine/Plugins/Experimental/Water/Source/Runtime/Private/WaterBodyRiverComponent.cpp
--- a/Engine/Plugins/Experimental/Water/Source/Runtime/Private/WaterBodyRiverComponent.cpp
+++ b/Engine/Plugins/Experimental/Water/Source/Runtime/Private/WaterBodyRiverComponent.cpp
@@ -517,12 +517,12 @@ void UWaterBodyRiverComponent::UpdateSplineMesh(USplineMeshComponent* MeshComp,
 }
 
 #if WITH_EDITOR
-void UWaterBodyRiverComponent::OnPostEditChangeProperty(FPropertyChangedEvent& PropertyChangedEvent, bool& bShapeOrPositionChanged, bool& bWeightmapSettingsChanged)
+void UWaterBodyRiverComponent::OnPostEditChangeProperty(FOnWaterBodyChangedParams& InOutOnWaterBodyChangedParams)
 {
-	Super::OnPostEditChangeProperty(PropertyChangedEvent, bShapeOrPositionChanged, bWeightmapSettingsChanged);
+	Super::OnPostEditChangeProperty(InOutOnWaterBodyChangedParams);
 
-	if (PropertyChangedEvent.GetPropertyName() == GET_MEMBER_NAME_CHECKED(UWaterBodyRiverComponent, LakeTransitionMaterial) ||
-		PropertyChangedEvent.GetPropertyName() == GET_MEMBER_NAME_CHECKED(UWaterBodyRiverComponent, OceanTransitionMaterial))
+	if (InOutOnWaterBodyChangedParams.PropertyChangedEvent.GetPropertyName() == GET_MEMBER_NAME_CHECKED(UWaterBodyRiverComponent, LakeTransitionMaterial) ||
+		InOutOnWaterBodyChangedParams.PropertyChangedEvent.GetPropertyName() == GET_MEMBER_NAME_CHECKED(UWaterBodyRiverComponent, OceanTransitionMaterial))
 	{
 		UpdateMaterialInstances();
 	}
diff --git a/Engine/Plugins/Experimental/Water/Source/Runtime/Private/WaterSplineComponent.cpp b/Engine/Plugins/Experimental/Water/Source/Runtime/Private/WaterSplineComponent.cpp
--- a/Engine/Plugins/Experimental/Water/Source/Runtime/Private/WaterSplineComponent.cpp
+++ b/Engine/Plugins/Experimental/Water/Source/Runtime/Private/WaterSplineComponent.cpp
@@ -42,7 +42,7 @@ void UWaterSplineComponent::PostLoad()
 /*
 	if (bAnythingChanged)
 	{
-		SplineDataChangedEvent.Broadcast();
+		WaterSplineDataChangedEvent.Broadcast();
 	}*/
 #endif
 }
@@ -56,7 +56,7 @@ void UWaterSplineComponent::PostDuplicate(bool bDuplicateForPie)
 	{
 		SynchronizeWaterProperties();
 
-		SplineDataChangedEvent.Broadcast();
+		WaterSplineDataChangedEvent.Broadcast(FOnWaterSplineDataChangedParams());
 	}
 #endif // WITH_EDITOR
 }
@@ -142,7 +142,7 @@ void UWaterSplineComponent::PostEditUndo()
 {
 	Super::PostEditUndo();
 
-	SplineDataChangedEvent.Broadcast();
+	WaterSplineDataChangedEvent.Broadcast(FOnWaterSplineDataChangedParams());
 }
 
 void UWaterSplineComponent::PostEditChangeProperty(FPropertyChangedEvent& PropertyChangedEvent)
@@ -151,7 +151,7 @@ void UWaterSplineComponent::PostEditChangeProperty(FPropertyChangedEvent& Proper
 
 	SynchronizeWaterProperties();
 
-	SplineDataChangedEvent.Broadcast();
+	WaterSplineDataChangedEvent.Broadcast(FOnWaterSplineDataChangedParams(PropertyChangedEvent));
 }
 
 void UWaterSplineComponent::PostEditImport()
@@ -160,7 +160,7 @@ void UWaterSplineComponent::PostEditImport()
 
 	SynchronizeWaterProperties();
 
-	SplineDataChangedEvent.Broadcast();
+	WaterSplineDataChangedEvent.Broadcast(FOnWaterSplineDataChangedParams());
 }
 
 void UWaterSplineComponent::ResetSpline(const TArray<FVector>& Points)
@@ -175,7 +175,7 @@ void UWaterSplineComponent::ResetSpline(const TArray<FVector>& Points)
 
 	UpdateSpline();
 	SynchronizeWaterProperties();
-	SplineDataChangedEvent.Broadcast();
+	WaterSplineDataChangedEvent.Broadcast(FOnWaterSplineDataChangedParams());
 }
 
 bool UWaterSplineComponent::SynchronizeWaterProperties()
diff --git a/Engine/Plugins/Experimental/Water/Source/Runtime/Private/WaterZoneActor.cpp b/Engine/Plugins/Experimental/Water/Source/Runtime/Private/WaterZoneActor.cpp
--- a/Engine/Plugins/Experimental/Water/Source/Runtime/Private/WaterZoneActor.cpp
+++ b/Engine/Plugins/Experimental/Water/Source/Runtime/Private/WaterZoneActor.cpp
@@ -109,14 +109,17 @@ void AWaterZone::MarkForRebuild(EWaterZoneRebuildFlags Flags)
 {
 	if (EnumHasAnyFlags(Flags, EWaterZoneRebuildFlags::UpdateWaterMesh))
 	{
+		UE_LOG(LogWater, Verbose, TEXT("AWaterZone::MarkForRebuild (UpdateWaterMesh)"));
 		WaterMesh->MarkWaterMeshGridDirty();
 	}
 	if (EnumHasAnyFlags(Flags, EWaterZoneRebuildFlags::UpdateWaterInfoTexture))
 	{
+		UE_LOG(LogWater, Verbose, TEXT("AWaterZone::MarkForRebuild (UpdateWaterInfoTexture)"));
 		bNeedsWaterInfoRebuild = true;
 	}
 	if (EnumHasAnyFlags(Flags, EWaterZoneRebuildFlags::UpdateWaterBodyLODSections))
 	{
+		UE_LOG(LogWater, Verbose, TEXT("AWaterZone::MarkForRebuild (UpdateWaterBodyLODSections)"));
 		ForEachWaterBodyComponent([](UWaterBodyComponent* WaterBodyComponent)
 		{
 			WaterBodyComponent->UpdateNonTessellatedMeshSections();
diff --git a/Engine/Plugins/Experimental/Water/Source/Runtime/Public/WaterBodyComponent.h b/Engine/Plugins/Experimental/Water/Source/Runtime/Public/WaterBodyComponent.h
--- a/Engine/Plugins/Experimental/Water/Source/Runtime/Public/WaterBodyComponent.h
+++ b/Engine/Plugins/Experimental/Water/Source/Runtime/Public/WaterBodyComponent.h
@@ -25,6 +25,7 @@
 #include "WaterBodyComponent.generated.h"
 
 class UWaterSplineComponent;
+struct FOnWaterSplineDataChangedParams;
 class AWaterBodyIsland;
 class AWaterBodyExclusionVolume;
 class AWaterZone;
@@ -104,6 +105,26 @@ public:
 	FBox2D SectionBounds;
 };
 
+
+// ----------------------------------------------------------------------------------
+
+struct FOnWaterBodyChangedParams
+{
+	FOnWaterBodyChangedParams(const FPropertyChangedEvent& InPropertyChangedEvent = FPropertyChangedEvent(/*InProperty = */nullptr))
+		: PropertyChangedEvent(InPropertyChangedEvent)
+	{}
+
+	/** Provides some additional context about how the water body data has changed (property, type of change...) */
+	FPropertyChangedEvent PropertyChangedEvent;
+
+	/** Indicates that property related to the water body's visual shape has changed */
+	bool bShapeOrPositionChanged = false;
+
+	/** Indicates that a property affecting the terrain weightmaps has changed */
+	bool bWeightmapSettingsChanged = false;
+};
+
+
 // ----------------------------------------------------------------------------------
 
 UCLASS(Abstract, HideCategories = (Tags, Activation, Cooking, Physics, Replication, Input, AssetUserData, Mesh))
@@ -119,7 +140,9 @@ public:
 	virtual bool AffectsWaterInfo() const;
 	virtual bool CanEverAffectWaterMesh() const { return true; }
 
-	void UpdateAll(bool bShapeOrPositionChanged);
+	UE_DEPRECATED(5.1, "Use the version of this function taking FOnWaterBodyChangedParams in parameter")
+	void UpdateAll(bool bShapeOrPositionChanged) {}
+	void UpdateAll(const FOnWaterBodyChangedParams& InParams);
 
 #if WITH_EDITOR
 	const FWaterCurveSettings& GetWaterCurveSettings() const { return CurveSettings; }
@@ -380,6 +403,7 @@ public:
 	
 	UE_DEPRECATED(5.1, "Renamed to CanEverAffectWaterMesh")
 	virtual bool CanAffectWaterMesh() const { return true; }
+
 protected:
 	//~ Begin UActorComponent interface.
 	virtual bool IsHLODRelevant() const override;
@@ -416,10 +440,12 @@ protected:
 	/** Returns navigation area class */
 	TSubclassOf<UNavAreaBase> GetNavAreaClass() const { return WaterNavAreaClass; }
 
+	/** Copies the relevant collision settings from the water body component to the component passed in parameter (useful when using external components for collision) */
 	void CopySharedCollisionSettingsToComponent(UPrimitiveComponent* InComponent);
+
+	/** Copies the relevant navigation settings from the water body component to the component passed in parameter (useful when using external components for navigation) */
 	void CopySharedNavigationSettingsToComponent(UPrimitiveComponent* InComponent);
 
-protected:
 	/** Computes the raw wave perturbation of the water height/normal */
 	virtual float GetWaveHeightAtPosition(const FVector& InPosition, float InWaterDepth, float InTime, FVector& OutNormal) const;
 
@@ -431,7 +457,9 @@ protected:
 
 #if WITH_EDITOR
 	/** Called by UWaterBodyComponent::PostEditChangeProperty. */
-	virtual void OnPostEditChangeProperty(FPropertyChangedEvent& PropertyChangedEvent, bool& bShapeOrPositionChanged, bool& bWeightmapSettingsChanged);
+	UE_DEPRECATED(5.1, "Use the version of the function taking FOnWaterBodyChangedParams")
+	virtual void OnPostEditChangeProperty(FPropertyChangedEvent& PropertyChangedEvent, bool& bShapeOrPositionChanged, bool& bWeightmapSettingsChanged) {}
+	virtual void OnPostEditChangeProperty(FOnWaterBodyChangedParams& InOutOnWaterBodyChangedParams);
 
 	/** Validates this component's data */
 	virtual void CheckForErrors() override;
@@ -463,6 +491,7 @@ protected:
 	void RebuildWaterBodyInfoMesh();
 	void RebuildWaterBodyLODSections();
 	void OnTessellatedWaterMeshBoundsChanged();
+	void OnWaterBodyChanged(const FOnWaterBodyChangedParams& InParams);
 
 	virtual void Serialize(FArchive& Ar) override;
 	virtual void PostLoad() override;
@@ -477,7 +506,9 @@ protected:
 	virtual void PostEditUndo() override;
 	virtual void PostEditImport() override;
 	virtual void PostEditChangeProperty(FPropertyChangedEvent& PropertyChangedEvent) override;
-	void OnSplineDataChanged();
+	UE_DEPRECATED(5.1, "Use OnWaterSplineDataChanged")
+	void OnSplineDataChanged() {}
+	void OnWaterSplineDataChanged(const FOnWaterSplineDataChangedParams& InParams);
 	void RegisterOnUpdateWavesData(UWaterWavesBase* InWaterWaves, bool bRegister);
 	void OnWavesDataUpdated(UWaterWavesBase* InWaterWaves, EPropertyChangeType::Type InChangeType);
 	void OnWaterSplineMetadataChanged(UWaterSplineMetadata* InWaterSplineMetadata, FPropertyChangedEvent& PropertyChangedEvent);
diff --git a/Engine/Plugins/Experimental/Water/Source/Runtime/Public/WaterBodyIslandActor.h b/Engine/Plugins/Experimental/Water/Source/Runtime/Public/WaterBodyIslandActor.h
--- a/Engine/Plugins/Experimental/Water/Source/Runtime/Public/WaterBodyIslandActor.h
+++ b/Engine/Plugins/Experimental/Water/Source/Runtime/Public/WaterBodyIslandActor.h
@@ -12,10 +12,32 @@
 #include "WaterBodyIslandActor.generated.h"
 
 class UWaterSplineComponent;
+struct FOnWaterSplineDataChangedParams;
 class USplineMeshComponent;
 class UBillboardComponent;
 class AWaterBody;
 
+// ----------------------------------------------------------------------------------
+
+struct FOnWaterBodyIslandChangedParams
+{
+	FOnWaterBodyIslandChangedParams(const FPropertyChangedEvent& InPropertyChangedEvent = FPropertyChangedEvent(/*InProperty = */nullptr))
+		: PropertyChangedEvent(InPropertyChangedEvent)
+	{}
+
+	/** Provides some additional context about how the water body island data has changed (property, type of change...) */
+	FPropertyChangedEvent PropertyChangedEvent;
+
+	/** Indicates that property related to the water body island's visual shape has changed */
+	bool bShapeOrPositionChanged = false;
+
+	/** Indicates that a property affecting the terrain weightmaps has changed */
+	bool bWeightmapSettingsChanged = false;
+};
+
+
+// ----------------------------------------------------------------------------------
+
 UCLASS(Blueprintable)
 class WATER_API AWaterBodyIsland : public AActor, public IWaterBrushActorInterface
 {
@@ -72,9 +94,19 @@ protected:
 	virtual void PostEditUndo() override;
 	virtual void PostEditChangeProperty(FPropertyChangedEvent& PropertyChangedEvent) override;
 	virtual void PostEditImport() override;
-	void UpdateAll();
-	void OnSplineDataChanged();
-	void OnWaterBodyIslandChanged(bool bShapeOrPositionChanged, bool bWeightmapSettingsChanged);
+	UE_DEPRECATED(5.1, "Use the version of this function taking FOnWaterBodyIslandChangedParams in parameter")
+	void UpdateAll() {}
+	void UpdateAll(const FOnWaterBodyIslandChangedParams& InParams);
+	
+	void OnWaterSplineDataChanged(const FOnWaterSplineDataChangedParams& InParams);
+	void OnWaterBodyIslandChanged(const FOnWaterBodyIslandChangedParams& InParams);
+
+	UE_DEPRECATED(5.1, "Use OnWaterSplineDataChanged")
+	void OnSplineDataChanged() {}
+
+	UE_DEPRECATED(5.1, "Use the version of this function taking a FOnWaterBodyIslandChangedParams in parameter")
+	void OnWaterBodyIslandChanged(bool bShapeOrPositionChanged, bool bWeightmapSettingsChanged) {}
+
 #endif
 
 protected:
diff --git a/Engine/Plugins/Experimental/Water/Source/Runtime/Public/WaterBodyOceanComponent.h b/Engine/Plugins/Experimental/Water/Source/Runtime/Public/WaterBodyOceanComponent.h
--- a/Engine/Plugins/Experimental/Water/Source/Runtime/Public/WaterBodyOceanComponent.h
+++ b/Engine/Plugins/Experimental/Water/Source/Runtime/Public/WaterBodyOceanComponent.h
@@ -39,7 +39,7 @@ protected:
 	virtual FBoxSphereBounds CalcBounds(const FTransform& LocalToWorld) const;
 
 #if WITH_EDITOR
-	virtual void OnPostEditChangeProperty(FPropertyChangedEvent& PropertyChangedEvent, bool& bShapeOrPositionChanged, bool& bWeightmapSettingsChanged) override;
+	virtual void OnPostEditChangeProperty(FOnWaterBodyChangedParams& InOutOnWaterBodyChangedParams) override;
 
 	virtual const TCHAR* GetWaterSpriteTextureName() const override;
 #endif
diff --git a/Engine/Plugins/Experimental/Water/Source/Runtime/Public/WaterBodyRiverComponent.h b/Engine/Plugins/Experimental/Water/Source/Runtime/Public/WaterBodyRiverComponent.h
--- a/Engine/Plugins/Experimental/Water/Source/Runtime/Public/WaterBodyRiverComponent.h
+++ b/Engine/Plugins/Experimental/Water/Source/Runtime/Public/WaterBodyRiverComponent.h
@@ -40,7 +40,7 @@ protected:
 
 	virtual FBoxSphereBounds CalcBounds(const FTransform& LocalToWorld) const override;
 #if WITH_EDITOR
-	virtual void OnPostEditChangeProperty(FPropertyChangedEvent& PropertyChangedEvent, bool& bShapeOrPositionChanged, bool& bWeightmapSettingsChanged) override;
+	virtual void OnPostEditChangeProperty(FOnWaterBodyChangedParams& InOutOnWaterBodyChangedParams) override;
 
 	virtual const TCHAR* GetWaterSpriteTextureName() const override;
 #endif
diff --git a/Engine/Plugins/Experimental/Water/Source/Runtime/Public/WaterBrushActorInterface.h b/Engine/Plugins/Experimental/Water/Source/Runtime/Public/WaterBrushActorInterface.h
--- a/Engine/Plugins/Experimental/Water/Source/Runtime/Public/WaterBrushActorInterface.h
+++ b/Engine/Plugins/Experimental/Water/Source/Runtime/Public/WaterBrushActorInterface.h
@@ -81,12 +81,21 @@ class WATER_API IWaterBrushActorInterface
 	 */
 	struct FWaterBrushActorChangedEventParams
 	{
-		FWaterBrushActorChangedEventParams(IWaterBrushActorInterface* InWaterBrushActor)
+		FWaterBrushActorChangedEventParams(IWaterBrushActorInterface* InWaterBrushActor, const FPropertyChangedEvent& InPropertyChangedEvent = FPropertyChangedEvent(/*InProperty = */nullptr))
 			: WaterBrushActor(InWaterBrushActor)
+			, PropertyChangedEvent(InPropertyChangedEvent)
 		{}
 
+		/** The water brush actor that has changed */
 		IWaterBrushActorInterface* WaterBrushActor = nullptr;
+
+		/** Provides some additional context about how the water brush actor data has changed (property, type of change...) */
+		FPropertyChangedEvent PropertyChangedEvent;
+
+		/** Indicates that property related to the water brush actor's visual shape has changed */
 		bool bShapeOrPositionChanged = false;
+
+		/** Indicates that a property affecting the terrain weightmaps has changed */
 		bool bWeightmapSettingsChanged = false;
 	};
 
diff --git a/Engine/Plugins/Experimental/Water/Source/Runtime/Public/WaterSplineComponent.h b/Engine/Plugins/Experimental/Water/Source/Runtime/Public/WaterSplineComponent.h
--- a/Engine/Plugins/Experimental/Water/Source/Runtime/Public/WaterSplineComponent.h
+++ b/Engine/Plugins/Experimental/Water/Source/Runtime/Public/WaterSplineComponent.h
@@ -8,6 +8,19 @@
 #include "WaterSplineMetadata.h"
 #include "WaterSplineComponent.generated.h"
 
+#if WITH_EDITOR
+struct FOnWaterSplineDataChangedParams
+{
+	FOnWaterSplineDataChangedParams(const FPropertyChangedEvent& InPropertyChangedEvent = FPropertyChangedEvent(/*InProperty = */nullptr)) 
+		: PropertyChangedEvent(InPropertyChangedEvent)
+	{}
+
+	/** Provides some additional context about how the water brush actor data has changed (property, type of change...) */
+	FPropertyChangedEvent PropertyChangedEvent;
+};
+DECLARE_MULTICAST_DELEGATE_OneParam(FOnWaterSplineDataChanged, const FOnWaterSplineDataChangedParams&);
+#endif // WITH_EDITOR
+
 UCLASS(ClassGroup = (Custom), meta = (BlueprintSpawnableComponent))
 class WATER_API UWaterSplineComponent : public USplineComponent
 {
@@ -38,8 +51,13 @@ public:
 	virtual bool AllowsSplinePointScaleEditing() const override { return false; }
 
 #if WITH_EDITOR
-	DECLARE_EVENT(UWaterSplineComponent, FOnSplineDataChanged);
-	FOnSplineDataChanged& OnSplineDataChanged() { return SplineDataChangedEvent; }
+	DECLARE_EVENT(UWaterSplineComponent, UE_DEPRECATED(5.1, "Use FOnWaterSplineDataChanged") FOnSplineDataChanged);
+	UE_DEPRECATED(5.1, "Use OnWaterSplineDataChanged")
+PRAGMA_DISABLE_DEPRECATION_WARNINGS
+	FOnSplineDataChanged& OnSplineDataChanged() { static FOnSplineDataChanged DeprecatedEvent; return DeprecatedEvent; }
+PRAGMA_ENABLE_DEPRECATION_WARNINGS
+
+	FOnWaterSplineDataChanged& OnWaterSplineDataChanged() { return WaterSplineDataChangedEvent; }
 
 	virtual bool CanEditChange(const FProperty* InProperty) const override;
 	virtual void PostEditUndo() override;
@@ -48,13 +66,13 @@ public:
 	
 	void ResetSpline(const TArray<FVector>& Points);
 	bool SynchronizeWaterProperties();
-#endif
+#endif // WITH_EDITOR
 
 	virtual FBoxSphereBounds CalcBounds(const FTransform& LocalToWorld) const override;
 	virtual void Serialize(FArchive& Ar) override;
 
 private:
 #if WITH_EDITOR
-	FOnSplineDataChanged SplineDataChangedEvent;
-#endif
+	FOnWaterSplineDataChanged WaterSplineDataChangedEvent;
+#endif // WITH_EDITOR
 };
\ No newline at end of file
diff --git a/Engine/Source/Editor/ComponentVisualizers/Private/SplineComponentVisualizer.cpp b/Engine/Source/Editor/ComponentVisualizers/Private/SplineComponentVisualizer.cpp
--- a/Engine/Source/Editor/ComponentVisualizers/Private/SplineComponentVisualizer.cpp
+++ b/Engine/Source/Editor/ComponentVisualizers/Private/SplineComponentVisualizer.cpp
@@ -1151,7 +1151,8 @@ bool FSplineComponentVisualizer::HandleInputDelta(FEditorViewportClient* Viewpor
 		check(SelectionState);
 		if (SelectionState->GetSelectedTangentHandle() != INDEX_NONE)
 		{
-			return TransformSelectedTangent(DeltaTranslate);
+			// Transform the tangent using an EPropertyChangeType::Interactive change. Later on, at the end of mouse tracking, a non-interactive change will be notified via void TrackingStopped :
+			return TransformSelectedTangent(EPropertyChangeType::Interactive, DeltaTranslate);
 		}
 		else if (ViewportClient->IsAltPressed())
 		{
@@ -1197,14 +1198,15 @@ bool FSplineComponentVisualizer::HandleInputDelta(FEditorViewportClient* Viewpor
 		}
 		else
 		{
-			return TransformSelectedKeys(DeltaTranslate, DeltaRotate, DeltaScale);
+			// Transform the spline keys using an EPropertyChangeType::Interactive change. Later on, at the end of mouse tracking, a non-interactive change will be notified via void TrackingStopped :
+			return TransformSelectedKeys(EPropertyChangeType::Interactive, DeltaTranslate, DeltaRotate, DeltaScale);
 		}
 	}
 
 	return false;
 }
 
-bool FSplineComponentVisualizer::TransformSelectedTangent(const FVector& DeltaTranslate)
+bool FSplineComponentVisualizer::TransformSelectedTangent(EPropertyChangeType::Type InPropertyChangeType, const FVector& InDeltaTranslate)
 {
 	USplineComponent* SplineComp = GetEditedSplineComponent();
 	if (SplineComp != nullptr)
@@ -1218,7 +1220,7 @@ bool FSplineComponentVisualizer::TransformSelectedTangent(const FVector& DeltaTr
 		ESelectedTangentHandle SelectedTangentHandleType;
 		SelectionState->GetVerifiedSelectedTangentHandle(NumPoints, SelectedTangentHandle, SelectedTangentHandleType);
 
-		if (!DeltaTranslate.IsZero())
+		if (!InDeltaTranslate.IsZero())
 		{
 			SplineComp->Modify();
 
@@ -1229,16 +1231,16 @@ bool FSplineComponentVisualizer::TransformSelectedTangent(const FVector& DeltaTr
 			{
 				if (SelectedTangentHandleType == ESelectedTangentHandle::Leave)
 				{
-					EditedPoint.LeaveTangent += SplineComp->GetComponentTransform().InverseTransformVector(DeltaTranslate) / TangentScale;
+					EditedPoint.LeaveTangent += SplineComp->GetComponentTransform().InverseTransformVector(InDeltaTranslate) / TangentScale;
 				}
 				else
 				{
-					EditedPoint.ArriveTangent += SplineComp->GetComponentTransform().InverseTransformVector(-DeltaTranslate) / TangentScale;
+					EditedPoint.ArriveTangent += SplineComp->GetComponentTransform().InverseTransformVector(-InDeltaTranslate) / TangentScale;
 				}
 			}
 			else
 			{
-				const FVector Delta = (SelectedTangentHandleType == ESelectedTangentHandle::Leave) ? DeltaTranslate : -DeltaTranslate;
+				const FVector Delta = (SelectedTangentHandleType == ESelectedTangentHandle::Leave) ? InDeltaTranslate : -InDeltaTranslate;
 				const FVector Tangent = EditedPoint.LeaveTangent + SplineComp->GetComponentTransform().InverseTransformVector(Delta) / TangentScale;
 
 				EditedPoint.LeaveTangent = Tangent;
@@ -1251,7 +1253,7 @@ bool FSplineComponentVisualizer::TransformSelectedTangent(const FVector& DeltaTr
 		SplineComp->UpdateSpline();
 		SplineComp->bSplineHasBeenEdited = true;
 
-		NotifyPropertyModified(SplineComp, SplineCurvesProperty);
+		NotifyPropertyModified(SplineComp, SplineCurvesProperty, InPropertyChangeType);
 
 		return true;
 	}
@@ -1259,7 +1261,7 @@ bool FSplineComponentVisualizer::TransformSelectedTangent(const FVector& DeltaTr
 	return false;
 }
 
-bool FSplineComponentVisualizer::TransformSelectedKeys(const FVector& DeltaTranslate, const FRotator& DeltaRotate, const FVector& DeltaScale)
+bool FSplineComponentVisualizer::TransformSelectedKeys(EPropertyChangeType::Type InPropertyChangeType, const FVector& InDeltaTranslate, const FRotator& InDeltaRotate, const FVector& InDeltaScale)
 {
 	USplineComponent* SplineComp = GetEditedSplineComponent();
 	if (SplineComp != nullptr)
@@ -1288,65 +1290,65 @@ bool FSplineComponentVisualizer::TransformSelectedKeys(const FVector& DeltaTrans
 			FInterpCurvePoint<FVector>& EditedScalePoint = SplineScale.Points[SelectedKeyIndex];
 
 
-			if (!DeltaTranslate.IsZero())
+			if (!InDeltaTranslate.IsZero())
 			{
 				// Find key position in world space
 				const FVector CurrentWorldPos = SplineComp->GetComponentTransform().TransformPosition(EditedPoint.OutVal);
 				// Move in world space
-				const FVector NewWorldPos = CurrentWorldPos + DeltaTranslate;
+				const FVector NewWorldPos = CurrentWorldPos + InDeltaTranslate;
 
 				// Convert back to local space
 				EditedPoint.OutVal = SplineComp->GetComponentTransform().InverseTransformPosition(NewWorldPos);
 			}
 
-			if (!DeltaRotate.IsZero())
+			if (!InDeltaRotate.IsZero())
 			{
 				// Set point tangent as user controlled
 				EditedPoint.InterpMode = CIM_CurveUser;
 
 				// Rotate tangent according to delta rotation
 				FVector NewTangent = SplineComp->GetComponentTransform().GetRotation().RotateVector(EditedPoint.LeaveTangent); // convert local-space tangent vector to world-space
-				NewTangent = DeltaRotate.RotateVector(NewTangent); // apply world-space delta rotation to world-space tangent
+				NewTangent = InDeltaRotate.RotateVector(NewTangent); // apply world-space delta rotation to world-space tangent
 				NewTangent = SplineComp->GetComponentTransform().GetRotation().Inverse().RotateVector(NewTangent); // convert world-space tangent vector back into local-space
 				EditedPoint.LeaveTangent = NewTangent;
 				EditedPoint.ArriveTangent = NewTangent;
 
 				// Rotate spline rotation according to delta rotation
 				FQuat NewRot = SplineComp->GetComponentTransform().GetRotation() * EditedRotPoint.OutVal; // convert local-space rotation to world-space
-				NewRot = DeltaRotate.Quaternion() * NewRot; // apply world-space rotation
+				NewRot = InDeltaRotate.Quaternion() * NewRot; // apply world-space rotation
 				NewRot = SplineComp->GetComponentTransform().GetRotation().Inverse() * NewRot; // convert world-space rotation to local-space
 				EditedRotPoint.OutVal = NewRot;
 			}
 
-			if (DeltaScale.X != 0.0f)
+			if (InDeltaScale.X != 0.0f)
 			{
 				// Set point tangent as user controlled
 				EditedPoint.InterpMode = CIM_CurveUser;
 
-				const FVector NewTangent = EditedPoint.LeaveTangent * (1.0f + DeltaScale.X);
+				const FVector NewTangent = EditedPoint.LeaveTangent * (1.0f + InDeltaScale.X);
 				EditedPoint.LeaveTangent = NewTangent;
 				EditedPoint.ArriveTangent = NewTangent;
 			}
 
-			if (DeltaScale.Y != 0.0f)
+			if (InDeltaScale.Y != 0.0f)
 			{
 				// Scale in Y adjusts the scale spline
-				EditedScalePoint.OutVal.Y *= (1.0f + DeltaScale.Y);
+				EditedScalePoint.OutVal.Y *= (1.0f + InDeltaScale.Y);
 			}
 
-			if (DeltaScale.Z != 0.0f)
+			if (InDeltaScale.Z != 0.0f)
 			{
 				// Scale in Z adjusts the scale spline
-				EditedScalePoint.OutVal.Z *= (1.0f + DeltaScale.Z);
+				EditedScalePoint.OutVal.Z *= (1.0f + InDeltaScale.Z);
 			}
 		}
 
 		SplineComp->UpdateSpline();
 		SplineComp->bSplineHasBeenEdited = true;
 
-		NotifyPropertyModified(SplineComp, SplineCurvesProperty);
+		NotifyPropertyModified(SplineComp, SplineCurvesProperty, InPropertyChangeType);
 
-		if (!DeltaRotate.IsZero())
+		if (!InDeltaRotate.IsZero())
 		{
 			SelectionState->Modify();
 			SelectionState->SetCachedRotation(SplineComp->GetQuaternionAtSplinePoint(LastKeyIndexSelected, ESplineCoordinateSpace::World));
@@ -1705,6 +1707,17 @@ bool FSplineComponentVisualizer::HandleSnapTo(const bool bInAlign, const bool bI
 	return false;
 }
 
+void FSplineComponentVisualizer::TrackingStopped(FEditorViewportClient* InViewportClient, bool bInDidMove)
+{
+	if (bInDidMove)
+	{
+		// After dragging, notify that the spline curves property has changed one last time, this time as a EPropertyChangeType::ValueSet :
+		USplineComponent* SplineComp = GetEditedSplineComponent();
+		check(SplineComp != nullptr);
+		NotifyPropertyModified(SplineComp, SplineCurvesProperty, EPropertyChangeType::ValueSet);
+	}
+}
+
 void FSplineComponentVisualizer::OnSnapKeyToNearestSplinePoint(ESplineComponentSnapMode InSnapMode)
 {
 	const FScopedTransaction Transaction(LOCTEXT("SnapToNearestSplinePoint", "Snap To Nearest Spline Point"));
@@ -2601,7 +2614,8 @@ void FSplineComponentVisualizer::UpdateSplitSegment(const FVector& InDrag)
 	SplineComp->UpdateSpline();
 	SplineComp->bSplineHasBeenEdited = true;
 
-	NotifyPropertyModified(SplineComp, SplineCurvesProperty);
+	// Transform the spline keys using an EPropertyChangeType::Interactive change. Later on, at the end of mouse tracking, a non-interactive change will be notified via void TrackingStopped :
+	NotifyPropertyModified(SplineComp, SplineCurvesProperty, EPropertyChangeType::Interactive);
 
 	GEditor->RedrawLevelEditingViewports(true);
 }
@@ -2730,7 +2744,8 @@ void FSplineComponentVisualizer::UpdateAddSegment(const FVector& InDrag)
 	SplineComp->UpdateSpline();
 	SplineComp->bSplineHasBeenEdited = true;
 
-	NotifyPropertyModified(SplineComp, SplineCurvesProperty);
+	// Transform the spline keys using an EPropertyChangeType::Interactive change. Later on, at the end of mouse tracking, a non-interactive change will be notified via void TrackingStopped :
+	NotifyPropertyModified(SplineComp, SplineCurvesProperty, EPropertyChangeType::Interactive);
 
 	GEditor->RedrawLevelEditingViewports(true);
 }
diff --git a/Engine/Source/Editor/ComponentVisualizers/Public/SplineComponentVisualizer.h b/Engine/Source/Editor/ComponentVisualizers/Public/SplineComponentVisualizer.h
--- a/Engine/Source/Editor/ComponentVisualizers/Public/SplineComponentVisualizer.h
+++ b/Engine/Source/Editor/ComponentVisualizers/Public/SplineComponentVisualizer.h
@@ -241,6 +241,8 @@ public:
 	virtual bool HasFocusOnSelectionBoundingBox(FBox& OutBoundingBox) override;
 	/** Pass snap input to active visualizer */
 	virtual bool HandleSnapTo(const bool bInAlign, const bool bInUseLineTrace, const bool bInUseBounds, const bool bInUsePivot, AActor* InDestination) override;
+	/** Gets called when the mouse tracking has stopped (dragging behavior) */
+	virtual void TrackingStopped(FEditorViewportClient* InViewportClient, bool bInDidMove) override;
 	/** Get currently edited component, this is needed to reset the active visualizer after undo/redo */
 	virtual UActorComponent* GetEditedComponent() const override;
 	virtual TSharedPtr<SWidget> GenerateContextMenu() const override;
@@ -284,10 +286,10 @@ protected:
 	void SelectSplinePoint(int32 SelectIndex, bool bAddToSelection);
 
 	/** Transforms selected tangent by given translation */
-	bool TransformSelectedTangent(const FVector& DeltaTranslate);
+	bool TransformSelectedTangent(EPropertyChangeType::Type InPropertyChangeType, const FVector& InDeltaTranslate);
 
 	/** Transforms selected tangent by given translate, rotate and scale */
-	bool TransformSelectedKeys(const FVector& DeltaTranslate, const FRotator& DeltaRotate = FRotator::ZeroRotator, const FVector& DeltaScale = FVector::ZeroVector);
+	bool TransformSelectedKeys(EPropertyChangeType::Type InPropertyChangeType, const FVector& InDeltaTranslate, const FRotator& InDeltaRotate = FRotator::ZeroRotator, const FVector& InDeltaScale = FVector::ZeroVector);
 
 	/** Update the key selection state of the visualizer */
 	virtual void ChangeSelectionState(int32 Index, bool bIsCtrlHeld);
diff --git a/Engine/Source/Editor/UnrealEd/Private/ComponentVisualizer.cpp b/Engine/Source/Editor/UnrealEd/Private/ComponentVisualizer.cpp
--- a/Engine/Source/Editor/UnrealEd/Private/ComponentVisualizer.cpp
+++ b/Engine/Source/Editor/UnrealEd/Private/ComponentVisualizer.cpp
@@ -253,7 +253,7 @@ void FComponentVisualizer::NotifyPropertiesModified(UActorComponent* Component,
 
 	for (FProperty* Property : Properties)
 	{
-		FPropertyChangedEvent PropertyChangedEvent(Property);
+		FPropertyChangedEvent PropertyChangedEvent(Property, PropertyChangeType);
 		Component->PostEditChangeProperty(PropertyChangedEvent);
 	}
 
diff --git a/Engine/Source/Editor/UnrealEd/Private/ComponentVisualizerManager.cpp b/Engine/Source/Editor/UnrealEd/Private/ComponentVisualizerManager.cpp
--- a/Engine/Source/Editor/UnrealEd/Private/ComponentVisualizerManager.cpp
+++ b/Engine/Source/Editor/UnrealEd/Private/ComponentVisualizerManager.cpp
@@ -218,6 +218,25 @@ bool FComponentVisualizerManager::GetCustomInputCoordinateSystem(const FEditorVi
 	return false;
 }
 
+void FComponentVisualizerManager::TrackingStarted(FEditorViewportClient* InViewportClient)
+{
+	TSharedPtr<FComponentVisualizer> EditedVisualizer = EditedVisualizerPtr.Pin();
+
+	if (EditedVisualizer.IsValid())
+	{
+		return EditedVisualizer->TrackingStarted(InViewportClient);
+	}
+}
+
+void FComponentVisualizerManager::TrackingStopped(FEditorViewportClient* InViewportClient, bool bInDidMove)
+{
+	TSharedPtr<FComponentVisualizer> EditedVisualizer = EditedVisualizerPtr.Pin();
+
+	if (EditedVisualizer.IsValid())
+	{
+		return EditedVisualizer->TrackingStopped(InViewportClient, bInDidMove);
+	}
+}
 
 TSharedPtr<SWidget> FComponentVisualizerManager::GenerateContextMenuForComponentVis() const
 {
diff --git a/Engine/Source/Editor/UnrealEd/Private/LevelEditorViewport.cpp b/Engine/Source/Editor/UnrealEd/Private/LevelEditorViewport.cpp
--- a/Engine/Source/Editor/UnrealEd/Private/LevelEditorViewport.cpp
+++ b/Engine/Source/Editor/UnrealEd/Private/LevelEditorViewport.cpp
@@ -3203,6 +3203,8 @@ void FLevelEditorViewportClient::TrackingStarted( const FInputEventState& InInpu
 			// Suspend actor/component modification during each delta step to avoid recording unnecessary overhead into the transaction buffer
 			GEditor->DisableDeltaModification(true);
 		}
+
+		GUnrealEd->ComponentVisManager.TrackingStarted(this);
 	}
 }
 
@@ -3262,6 +3264,8 @@ void FLevelEditorViewportClient::TrackingStopped()
 		{
 			GUnrealEd->UpdatePivotLocationForSelection();
 		}
+
+		GUnrealEd->ComponentVisManager.TrackingStopped(this, bDidMove);
 	}
 
 	if (bNeedToRestoreComponentBeingMovedFlag)
diff --git a/Engine/Source/Editor/UnrealEd/Public/ComponentVisualizer.h b/Engine/Source/Editor/UnrealEd/Public/ComponentVisualizer.h
--- a/Engine/Source/Editor/UnrealEd/Public/ComponentVisualizer.h
+++ b/Engine/Source/Editor/UnrealEd/Public/ComponentVisualizer.h
@@ -182,6 +182,10 @@ public:
 	virtual bool HasFocusOnSelectionBoundingBox(FBox& OutBoundingBox) { return false; }
 	/** Pass snap input to active visualizer */
 	virtual bool HandleSnapTo(const bool bInAlign, const bool bInUseLineTrace, const bool bInUseBounds, const bool bInUsePivot, AActor* InDestination) { return false;  }
+	/** Gets called when the mouse tracking has started (dragging behavior) */
+	virtual void TrackingStarted(FEditorViewportClient* InViewportClient) {}
+	/** Gets called when the mouse tracking has stopped (dragging behavior) */
+	virtual void TrackingStopped(FEditorViewportClient* InViewportClient, bool bInDidMove) {}
 	/** Get currently edited component, this is needed to reset the active visualizer after undo/redo */
 	virtual UActorComponent* GetEditedComponent() const { return nullptr;  }
 
diff --git a/Engine/Source/Editor/UnrealEd/Public/ComponentVisualizerManager.h b/Engine/Source/Editor/UnrealEd/Public/ComponentVisualizerManager.h
--- a/Engine/Source/Editor/UnrealEd/Public/ComponentVisualizerManager.h
+++ b/Engine/Source/Editor/UnrealEd/Public/ComponentVisualizerManager.h
@@ -55,6 +55,12 @@ public:
 	/** Get custom widget coordinate system from active visualizer */
 	bool GetCustomInputCoordinateSystem(const FEditorViewportClient* InViewportClient, FMatrix& OutMatrix) const;
 
+	/** Gets called when the mouse tracking has started (dragging behavior) */
+	void TrackingStarted(FEditorViewportClient* InViewportClient);
+
+	/** Gets called when the mouse tracking has stopped (dragging behavior) */
+	void TrackingStopped(FEditorViewportClient* InViewportClient, bool bInDidMove);
+
 	/** Generate context menu for the component visualizer */
 	TSharedPtr<SWidget> GenerateContextMenuForComponentVis() const;
 
diff --git a/Engine/Source/Runtime/Landscape/Private/LandscapeBlueprintBrushBase.cpp b/Engine/Source/Runtime/Landscape/Private/LandscapeBlueprintBrushBase.cpp
--- a/Engine/Source/Runtime/Landscape/Private/LandscapeBlueprintBrushBase.cpp
+++ b/Engine/Source/Runtime/Landscape/Private/LandscapeBlueprintBrushBase.cpp
@@ -4,6 +4,7 @@
 #include "CoreMinimal.h"
 #include "LandscapeProxy.h"
 #include "Landscape.h"
+#include "LandscapePrivate.h"
 #include "Misc/MapErrors.h"
 #include "Misc/UObjectToken.h"
 #include "Logging/MessageLog.h"
@@ -59,6 +60,7 @@ void ALandscapeBlueprintBrushBase::Initialize_Implementation(const FTransform& I
 void ALandscapeBlueprintBrushBase::RequestLandscapeUpdate()
 {
 #if WITH_EDITORONLY_DATA
+	UE_LOG(LogLandscape, Verbose, TEXT("ALandscapeBlueprintBrushBase::RequestLandscapeUpdate"));
 	if (OwningLandscape)
 	{
 		uint32 ModeMask = 0;
