switch the component visualization to use the physics object interface instead

diff --git a/Engine/Source/Editor/ComponentVisualizers/Private/ConstraintComponentVisualizer.cpp b/Engine/Source/Editor/ComponentVisualizers/Private/ConstraintComponentVisualizer.cpp
--- a/Engine/Source/Editor/ComponentVisualizers/Private/ConstraintComponentVisualizer.cpp
+++ b/Engine/Source/Editor/ComponentVisualizers/Private/ConstraintComponentVisualizer.cpp
@@ -42,7 +42,7 @@ void FConstraintComponentVisualizer::DrawVisualization( const UActorComponent* C
 
 			const float LastKnownScale = Instance.GetLastKnownScale();
 
-			if(ConstraintComp->GetBodyInstance(EConstraintFrame::Frame1))
+			if(ConstraintComp->GetPhysicsObject(EConstraintFrame::Frame1))
 			{
 				Con1Frame.ScaleTranslation(LastKnownScale);
 			}
@@ -50,7 +50,7 @@ void FConstraintComponentVisualizer::DrawVisualization( const UActorComponent* C
 
 			Con2Frame.SetScale3D(Con2Frame.GetScale3D() * LastKnownScale);
 
-			if (ConstraintComp->GetBodyInstance(EConstraintFrame::Frame2))
+			if (ConstraintComp->GetPhysicsObject(EConstraintFrame::Frame2))
 			{
 				Con2Frame.ScaleTranslation(LastKnownScale);
 			}
