Splnes: Added checks for valid widget path

diff --git a/Engine/Source/Editor/ComponentVisualizers/Private/SplineGeneratorPanel.cpp b/Engine/Source/Editor/ComponentVisualizers/Private/SplineGeneratorPanel.cpp
--- a/Engine/Source/Editor/ComponentVisualizers/Private/SplineGeneratorPanel.cpp
+++ b/Engine/Source/Editor/ComponentVisualizers/Private/SplineGeneratorPanel.cpp
@@ -642,7 +642,7 @@ void SSplineGeneratorPanel::OnFocusChanging(const FWeakWidgetPath& PreviousFocus
 {
 	FWidgetPath WidgetPath = NewWidgetPath.GetPathDownTo(SharedThis(this));
 	FWidgetPath WidgetPath2 = NewWidgetPath;
-	bool bPathFound = WidgetPath2.ExtendPathTo(FWidgetMatcher(SharedThis(this)));
+	bool bPathFound = WidgetPath2.IsValid() && WidgetPath2.ExtendPathTo(FWidgetMatcher(SharedThis(this)));
 	if (SplineGen && !WidgetPath.IsValid() && bPathFound)
 	{
 		SplineGen = nullptr;
