[NO NON-PII COMMIT MESSAGE]

diff --git a/Engine/Source/Editor/ComponentVisualizers/Private/SplineComponentVisualizer.cpp b/Engine/Source/Editor/ComponentVisualizers/Private/SplineComponentVisualizer.cpp
--- a/Engine/Source/Editor/ComponentVisualizers/Private/SplineComponentVisualizer.cpp
+++ b/Engine/Source/Editor/ComponentVisualizers/Private/SplineComponentVisualizer.cpp
@@ -99,7 +99,7 @@ public:
 	{
 		UI_COMMAND(DeleteKey, "Delete Spline Point", "Delete the currently selected spline point.", EUserInterfaceActionType::Button, FInputChord(EKeys::Delete));
 		UI_COMMAND(DuplicateKey, "Duplicate Spline Point", "Duplicate the currently selected spline point.", EUserInterfaceActionType::Button, FInputChord());
-		UI_COMMAND(AddKey, "Add Spline Point Here", "Add a new spline point at the cursor location.", EUserInterfaceActionType::Button, FInputChord());
+		UI_COMMAND(AddKey, "Add Spline Point Here", "Add a new spline point at the cursor location.", EUserInterfaceActionType::Button, FInputChord(EKeys::A));
 		UI_COMMAND(SelectAll, "Select All Spline Points", "Select all spline points.", EUserInterfaceActionType::Button, FInputChord());
 		UI_COMMAND(SelectNextSplinePoint, "Select Next Spline Point", "Select next spline point.", EUserInterfaceActionType::Button, FInputChord(EKeys::Period));
 		UI_COMMAND(SelectPrevSplinePoint, "Select Prev Spline Point", "Select prev spline point.", EUserInterfaceActionType::Button, FInputChord(EKeys::Comma));
