Fixed non unity build error

diff --git a/Engine/Source/Editor/ComponentVisualizers/Private/LocalHeightFogComponentVisualizer.cpp b/Engine/Source/Editor/ComponentVisualizers/Private/LocalHeightFogComponentVisualizer.cpp
--- a/Engine/Source/Editor/ComponentVisualizers/Private/LocalHeightFogComponentVisualizer.cpp
+++ b/Engine/Source/Editor/ComponentVisualizers/Private/LocalHeightFogComponentVisualizer.cpp
@@ -4,6 +4,8 @@
 
 #include "Components/ActorComponent.h"
 #include "Components/LocalHeightFogComponent.h"
+#include "SceneView.h"
+#include "SceneManagement.h"
 
 void FLocalHeightFogComponentVisualizer::DrawVisualization( const UActorComponent* Component, const FSceneView* View, FPrimitiveDrawInterface* PDI )
 {
