[Splines] Rework component visualizer visibility to key off draw debug variable.

When this is disabled, the component visualizer will still draw if the component is directly selected, and will allow editing this way.

Supports "helper" splines for PCG workflows that are intended to be hidden by default to the user.

diff --git a/Engine/Source/Editor/ComponentVisualizers/Private/SplineComponentVisualizer.cpp b/Engine/Source/Editor/ComponentVisualizers/Private/SplineComponentVisualizer.cpp
--- a/Engine/Source/Editor/ComponentVisualizers/Private/SplineComponentVisualizer.cpp
+++ b/Engine/Source/Editor/ComponentVisualizers/Private/SplineComponentVisualizer.cpp
@@ -455,6 +455,16 @@ void FSplineComponentVisualizer::OnRegister()
 		);
 }
 
+bool FSplineComponentVisualizer::ShouldShowForSelectedSubcomponents(const UActorComponent* Component)
+{
+	if (const USplineComponent* SplineComp = Cast<const USplineComponent>(Component))
+	{
+		return SplineComp->bDrawDebug;
+	}
+
+	return false;
+}
+
 FSplineComponentVisualizer::~FSplineComponentVisualizer()
 {
 	FSplineComponentVisualizerCommands::Unregister();
diff --git a/Engine/Source/Editor/ComponentVisualizers/Public/SplineComponentVisualizer.h b/Engine/Source/Editor/ComponentVisualizers/Public/SplineComponentVisualizer.h
--- a/Engine/Source/Editor/ComponentVisualizers/Public/SplineComponentVisualizer.h
+++ b/Engine/Source/Editor/ComponentVisualizers/Public/SplineComponentVisualizer.h
@@ -222,6 +222,7 @@ public:
 
 	//~ Begin FComponentVisualizer Interface
 	virtual void OnRegister() override;
+	virtual bool ShouldShowForSelectedSubcomponents(const UActorComponent* Component) override;
 	virtual void DrawVisualization(const UActorComponent* Component, const FSceneView* View, FPrimitiveDrawInterface* PDI) override;
 	virtual bool VisProxyHandleClick(FEditorViewportClient* InViewportClient, HComponentVisProxy* VisProxy, const FViewportClick& Click) override;
 	/** Draw HUD on viewport for the supplied component */
diff --git a/Engine/Source/Editor/UnrealEd/Private/UnrealEdEngine.cpp b/Engine/Source/Editor/UnrealEd/Private/UnrealEdEngine.cpp
--- a/Engine/Source/Editor/UnrealEd/Private/UnrealEdEngine.cpp
+++ b/Engine/Source/Editor/UnrealEd/Private/UnrealEdEngine.cpp
@@ -1493,7 +1493,7 @@ void UUnrealEdEngine::OnEditorSelectionChanged(UObject* SelectionThatChanged)
 			{
 				// Try and find a visualizer
 				TSharedPtr<FComponentVisualizer> Visualizer = FindComponentVisualizer(Comp->GetClass());
-				if (Visualizer.IsValid())
+				if (Visualizer.IsValid() && (Comp == SelectedComponent || Visualizer->ShouldShowForSelectedSubcomponents(Comp)))
 				{
 					FCachedComponentVisualizer CachedComponentVisualizer(Comp, Visualizer);
 					FComponentVisualizerForSelection Temp{CachedComponentVisualizer};
diff --git a/Engine/Source/Editor/UnrealEd/Public/ComponentVisualizer.h b/Engine/Source/Editor/UnrealEd/Public/ComponentVisualizer.h
--- a/Engine/Source/Editor/UnrealEd/Public/ComponentVisualizer.h
+++ b/Engine/Source/Editor/UnrealEd/Public/ComponentVisualizer.h
@@ -156,6 +156,8 @@ public:
 	virtual void OnRegister() {}
 	/** Only show this visualizer if the actor is selected */
 	virtual bool ShowWhenSelected() { return true; }
+	/** Show this visualizer if the component is directly is selected */
+	virtual bool ShouldShowForSelectedSubcomponents(const UActorComponent* Component) { return true; }
 	/** Draw visualization for the supplied component */
 	virtual void DrawVisualization(const UActorComponent* Component, const FSceneView* View, FPrimitiveDrawInterface* PDI) {}
 	/** Draw HUD on viewport for the supplied component */
