SplineComponent: Support surface snapping when dragging a single spline point by the ball of the gizmo.

diff --git a/Engine/Source/Editor/ComponentVisualizers/Private/SplineComponentVisualizer.cpp b/Engine/Source/Editor/ComponentVisualizers/Private/SplineComponentVisualizer.cpp
--- a/Engine/Source/Editor/ComponentVisualizers/Private/SplineComponentVisualizer.cpp
+++ b/Engine/Source/Editor/ComponentVisualizers/Private/SplineComponentVisualizer.cpp
@@ -3,11 +3,15 @@
 #include "SplineComponentVisualizer.h"
 #include "CoreMinimal.h"
 #include "Algo/AnyOf.h"
+#include "Components/BrushComponent.h"
+#include "CollisionQueryParams.h"
+#include "Engine/HitResult.h"
 #include "Framework/Application/SlateApplication.h"
 #include "Framework/Commands/InputChord.h"
 #include "Framework/Commands/Commands.h"
 #include "Framework/Commands/UICommandList.h"
 #include "Framework/MultiBox/MultiBoxBuilder.h"
+#include "GameFramework/Volume.h"
 #include "Styling/AppStyle.h"
 #include "UnrealWidgetFwd.h"
 #include "Editor.h"
@@ -38,6 +42,65 @@ DEFINE_LOG_CATEGORY_STATIC(LogSplineComponentVisualizer, Log, All)
 
 #define VISUALIZE_SPLINE_UPVECTORS 0
 
+namespace SplineComponentVisualizerLocals
+{
+	// This is mostly modeled on FindNearestVisibleObjectHit_Internal in ModelingSceneSnappingManager.cpp,
+	// which we probably can't access because it lives in a plugin.
+	// TODO: Perhaps this code should live in some common utility place? Certainly if we do this again.
+	bool RaycastWorld(const UWorld* World, FEditorViewportClient* ViewportClient, FViewport* Viewport, FHitResult& HitResultOut)
+	{
+		FSceneViewFamilyContext ViewFamily(FSceneViewFamily::ConstructionValues(
+			ViewportClient->Viewport,
+			ViewportClient->GetScene(),
+			ViewportClient->EngineShowFlags));
+		// this View is deleted by the FSceneViewFamilyContext destructor
+		FSceneView* View = ViewportClient->CalcSceneView(&ViewFamily);
+
+		FViewportCursorLocation MouseViewportRay(View, ViewportClient, Viewport->GetMouseX(), Viewport->GetMouseY());
+		FRay Ray(MouseViewportRay.GetOrigin(), MouseViewportRay.GetDirection());
+
+		FCollisionObjectQueryParams ObjectQueryParams(FCollisionObjectQueryParams::AllObjects);
+		FCollisionQueryParams QueryParams = FCollisionQueryParams::DefaultQueryParam;
+		QueryParams.bTraceComplex = true;
+
+		TArray<FHitResult> OutHits;
+		FVector RayEnd = (FVector)(Ray.PointAt(HALF_WORLD_MAX));
+		if (World->LineTraceMultiByObjectType(OutHits, (FVector)Ray.Origin, RayEnd, ObjectQueryParams, QueryParams) == false)
+		{
+			return false;
+		}
+
+		double NearestVisible = TNumericLimits<double>::Max();
+		for (const FHitResult& CurResult : OutHits)
+		{
+			UPrimitiveComponent* Component = CurResult.Component.Get();
+			AActor* Actor = CurResult.GetActor();
+
+			// Don't use volumes
+			if (Cast<UBrushComponent>(Component) != nullptr &&
+				Cast<AVolume>(Actor) != nullptr)
+			{
+				continue;
+			}
+
+			// Ignore invisible things
+			if ((Actor && Actor->IsHidden())
+				|| (Component && !Component->IsVisibleInEditor()))
+			{
+				continue;
+			}
+
+			if (CurResult.Distance < NearestVisible)
+			{
+				HitResultOut = CurResult;
+				NearestVisible = CurResult.Distance;
+			}
+		}
+
+		return NearestVisible < TNumericLimits<double>::Max();
+	}
+}
+
 int32 USplineComponentVisualizerSelectionState::GetVerifiedLastKeyIndexSelected(const int32 InNumSplinePoints) const
 {
 	check(LastKeyIndexSelected != INDEX_NONE);
@@ -1147,8 +1210,10 @@ bool FSplineComponentVisualizer::AreKeysSelected() const
 	return (IsSingleKeySelected() || AreMultipleKeysSelected());
 }
 
-bool FSplineComponentVisualizer::HandleInputDelta(FEditorViewportClient* ViewportClient, FViewport* Viewport, FVector& DeltaTranslate, FRotator& DeltaRotate, FVector& DeltaScale)
+bool FSplineComponentVisualizer::HandleInputDelta(FEditorViewportClient* ViewportClient, FViewport* Viewport, FVector& DeltaTranslateIn, FRotator& DeltaRotate, FVector& DeltaScale)
 {
+	using namespace SplineComponentVisualizerLocals;
+
 	ResetTempModes();
 
 	USplineComponent* SplineComp = GetEditedSplineComponent();
@@ -1162,6 +1227,30 @@ bool FSplineComponentVisualizer::HandleInputDelta(FEditorViewportClient* Viewpor
 		}
 
 		check(SelectionState);
+
+		// Use a local value for DeltaTranslate so that we can modify it based on the SnapToSurface setting without
+		// changing it for the caller (that parameter should probably be const, but it's a base class function)
+		FVector DeltaTranslate = DeltaTranslateIn;
+
+		// It's tough to port the actor "surface snapping" toggle behavior here because the original code only works on actors
+		// and involves things like keeping track of last transform, snapping to surface even if moving in a different plane, etc
+		// (see FLevelEditorViewportClient::ProjectActorsIntoWorld). However we can at least support the limited but useful case
+		// of trying to drag points around by the ball of the widget.
+		const ULevelEditorViewportSettings* ViewportSettings = GetDefault<ULevelEditorViewportSettings>();
+		const float SnapOffsetExtent = (ViewportSettings->SnapToSurface.bEnabled) ? (ViewportSettings->SnapToSurface.SnapOffsetExtent) : (0.0f);
+		if (ViewportSettings->SnapToSurface.bEnabled
+			&& ViewportClient->GetWidgetMode() == UE::Widget::EWidgetMode::WM_Translate
+			&& (ViewportClient->GetCurrentWidgetAxis() == EAxisList::Screen || ViewportClient->GetCurrentWidgetAxis() == EAxisList::XYZ)
+			&& !ViewportClient->IsOrtho())
+		{
+			FHitResult HitResult;
+			if (RaycastWorld(SplineComp->GetWorld(), ViewportClient, Viewport, HitResult))
+			{
+				FVector3d NewGizmoLocation = HitResult.ImpactPoint + ViewportSettings->SnapToSurface.SnapOffsetExtent * HitResult.ImpactNormal;
+				DeltaTranslate = NewGizmoLocation - ViewportClient->GetWidgetLocation();
+			}
+		}
+
 		if (SelectionState->GetSelectedTangentHandle() != INDEX_NONE)
 		{
 			// Transform the tangent using an EPropertyChangeType::Interactive change. Later on, at the end of mouse tracking, a non-interactive change will be notified via void TrackingStopped :
