LFV - force uniform scaleing as of today while we figure out the problem with non uniform scaling.

diff --git a/Engine/Source/Editor/ComponentVisualizers/Private/LocalFogVolumeComponentVisualizer.cpp b/Engine/Source/Editor/ComponentVisualizers/Private/LocalFogVolumeComponentVisualizer.cpp
--- a/Engine/Source/Editor/ComponentVisualizers/Private/LocalFogVolumeComponentVisualizer.cpp
+++ b/Engine/Source/Editor/ComponentVisualizers/Private/LocalFogVolumeComponentVisualizer.cpp
@@ -16,6 +16,9 @@ void FLocalFogVolumeComponentVisualizer::DrawVisualization( const UActorComponen
 		{
 			FTransform LocalFogVolumeTransform = LocalFogVolume->GetComponentTransform();
 
+			const float MaximumAxisScale = LocalFogVolumeTransform.GetMaximumAxisScale();
+			LocalFogVolumeTransform.SetScale3D(FVector(MaximumAxisScale, MaximumAxisScale, MaximumAxisScale));
+
 			// Draw local fog volume spherical shape radius
 			DrawWireSphereAutoSides(PDI, LocalFogVolumeTransform, FColor(200, 255, 255), 1.0f, SDPG_World);
 		}
diff --git a/Engine/Source/Runtime/Engine/Private/Components/LocalFogVolumeComponent.cpp b/Engine/Source/Runtime/Engine/Private/Components/LocalFogVolumeComponent.cpp
--- a/Engine/Source/Runtime/Engine/Private/Components/LocalFogVolumeComponent.cpp
+++ b/Engine/Source/Runtime/Engine/Private/Components/LocalFogVolumeComponent.cpp
@@ -42,6 +42,9 @@ void ULocalFogVolumeComponent::SendRenderTransformCommand()
 			{
 				// Nothing else is needed so that command could actually go.
 				SceneProxy->FogTransform = ComponentTransform;
+
+				const float MaximumAxisScale = SceneProxy->FogTransform.GetMaximumAxisScale();
+				SceneProxy->FogTransform.SetScale3D(FVector(MaximumAxisScale, MaximumAxisScale, MaximumAxisScale));
 			});
 	}
 }
