Added localization macro for Spline Generation panel title.

diff --git a/Engine/Source/Editor/ComponentVisualizers/Private/SplineComponentVisualizer.cpp b/Engine/Source/Editor/ComponentVisualizers/Private/SplineComponentVisualizer.cpp
--- a/Engine/Source/Editor/ComponentVisualizers/Private/SplineComponentVisualizer.cpp
+++ b/Engine/Source/Editor/ComponentVisualizers/Private/SplineComponentVisualizer.cpp
@@ -3556,7 +3556,7 @@ void FSplineComponentVisualizer::CreateSplineGeneratorPanel()
 	{
 		ExistingWindow = SNew(SWindow)
 			.ScreenPosition(FSlateApplication::Get().GetCursorPos())
-			.Title(FText::FromString("Spline Generation"))
+			.Title(LOCTEXT("SplineGenerationPanelTitle", "Spline Generation"))
 			.SizingRule(ESizingRule::Autosized)
 			.AutoCenter(EAutoCenter::None)
 			.SupportsMaximize(false)
