[SplineComponentVisualizer]
* Added new commands:
- ToggleSnapTangentAdjustments, StraightenToNext, StraightenToPrevious, ToggleClosedLoop

diff --git a/Engine/Source/Editor/ComponentVisualizers/Private/SplineComponentVisualizer.cpp b/Engine/Source/Editor/ComponentVisualizers/Private/SplineComponentVisualizer.cpp
--- a/Engine/Source/Editor/ComponentVisualizers/Private/SplineComponentVisualizer.cpp
+++ b/Engine/Source/Editor/ComponentVisualizers/Private/SplineComponentVisualizer.cpp
@@ -180,18 +180,22 @@ public:
 		UI_COMMAND(SnapKeyToActor, "Snap to Actor", "Snap selected spline point to actor, Ctrl-LMB to select the actor after choosing this option.", EUserInterfaceActionType::Button, FInputChord(EKeys::P, (EModifierKey::Alt | EModifierKey::Shift)));
 		UI_COMMAND(AlignKeyToActor, "Align to Actor", "Align selected spline point to actor, Ctrl-LMB to select the actor after choosing this option.", EUserInterfaceActionType::Button, FInputChord());
 		UI_COMMAND(AlignKeyPerpendicularToActor, "Align Perpendicular to Actor", "Align perpendicular  selected spline point to actor, Ctrl-LMB to select the actor after choosing this option.", EUserInterfaceActionType::Button, FInputChord());
+		UI_COMMAND(ToggleSnapTangentAdjustments, "Allow Tangents Updates On Snap", "Allow tangents to update when performing snap operations on points.", EUserInterfaceActionType::ToggleButton, FInputChord());
 		UI_COMMAND(SnapAllToSelectedX, "Snap All To Selected X", "Snap all spline points to selected spline point world X position.", EUserInterfaceActionType::Button, FInputChord());
 		UI_COMMAND(SnapAllToSelectedY, "Snap All To Selected Y", "Snap all spline points to selected spline point world Y position.", EUserInterfaceActionType::Button, FInputChord());
 		UI_COMMAND(SnapAllToSelectedZ, "Snap All To Selected Z", "Snap all spline points to selected spline point world Z position.", EUserInterfaceActionType::Button, FInputChord());
 		UI_COMMAND(SnapToLastSelectedX, "Snap To Last Selected X", "Snap selected spline points to world X position of last selected spline point.", EUserInterfaceActionType::Button, FInputChord());
 		UI_COMMAND(SnapToLastSelectedY, "Snap To Last Selected Y", "Snap selected spline points to world Y position of last selected spline point.", EUserInterfaceActionType::Button, FInputChord());
 		UI_COMMAND(SnapToLastSelectedZ, "Snap To Last Selected Z", "Snap selected spline points to world Z position of last selected spline point.", EUserInterfaceActionType::Button, FInputChord());
+		UI_COMMAND(StraightenToNext, "Straighten To Next Point", "Straighten selected points toward next sequential point", EUserInterfaceActionType::Button, FInputChord());
+		UI_COMMAND(StraightenToPrevious, "Straighten To Previous Point", "Straighten selected points toward previous sequential point", EUserInterfaceActionType::Button, FInputChord());
 		UI_COMMAND(SetLockedAxisNone, "None", "New spline point axis is not fixed.", EUserInterfaceActionType::RadioButton, FInputChord());
 		UI_COMMAND(SetLockedAxisX, "X", "Fix X axis when adding new spline points.", EUserInterfaceActionType::RadioButton, FInputChord());
 		UI_COMMAND(SetLockedAxisY, "Y", "Fix Y axis when adding new spline points.", EUserInterfaceActionType::RadioButton, FInputChord());
 		UI_COMMAND(SetLockedAxisZ, "Z", "Fix Z axis when adding new spline points.", EUserInterfaceActionType::RadioButton, FInputChord());
 		UI_COMMAND(VisualizeRollAndScale, "Visualize Roll and Scale", "Whether the visualization should show roll and scale on this spline.", EUserInterfaceActionType::ToggleButton, FInputChord());
 		UI_COMMAND(DiscontinuousSpline, "Allow Discontinuous Splines", "Whether the visualization allows Arrive and Leave tangents to be set separately.", EUserInterfaceActionType::ToggleButton, FInputChord());
+		UI_COMMAND(ToggleClosedLoop, "Closed Loop", "Toggle the Closed Loop setting of the spline", EUserInterfaceActionType::ToggleButton, FInputChord());
 		UI_COMMAND(ResetToDefault, "Reset to Default", "Reset this spline to its archetype default.", EUserInterfaceActionType::Button, FInputChord());
 	}
 
@@ -256,6 +260,9 @@ public:
 	/** Align key perpendicular to nearest actor */
 	TSharedPtr<FUICommandInfo> AlignKeyPerpendicularToActor;
 
+	/** Turn On / Off Tangent updates when snapping points*/
+	TSharedPtr<FUICommandInfo> ToggleSnapTangentAdjustments;
+
 	/** Snap all spline points to selected point world X position*/
 	TSharedPtr<FUICommandInfo> SnapAllToSelectedX;
 
@@ -274,6 +281,12 @@ public:
 	/** Snap selected spline points to last selected point world Z position */
 	TSharedPtr<FUICommandInfo> SnapToLastSelectedZ;
 
+	/** Straighten tangents to align directly toward Next spline points */
+	TSharedPtr<FUICommandInfo> StraightenToNext;
+
+	/** Straighten tangents to align directly toward Previous spline points */
+	TSharedPtr<FUICommandInfo> StraightenToPrevious;
+
 	/** No axis is locked when adding new spline points */
 	TSharedPtr<FUICommandInfo> SetLockedAxisNone;
 
@@ -292,6 +305,9 @@ public:
 	/** Whether we allow separate Arrive / Leave tangents, resulting in a discontinuous spline */
 	TSharedPtr<FUICommandInfo> DiscontinuousSpline;
 
+	/** Toggle the Closed Loop setting of the spline */
+	TSharedPtr<FUICommandInfo> ToggleClosedLoop;
+
 	/** Reset this spline to its default */
 	TSharedPtr<FUICommandInfo> ResetToDefault;
 };
@@ -425,6 +441,12 @@ void FSplineComponentVisualizer::OnRegister()
 		Commands.AlignKeyPerpendicularToActor,
 		FExecuteAction::CreateSP(this, &FSplineComponentVisualizer::OnSnapKeyToActor, ESplineComponentSnapMode::AlignPerpendicularToTangent),
 		FCanExecuteAction::CreateSP(this, &FSplineComponentVisualizer::IsSingleKeySelected));
+	
+	SplineComponentVisualizerActions->MapAction(
+		Commands.ToggleSnapTangentAdjustments,
+		FExecuteAction::CreateSP(this, &FSplineComponentVisualizer::OnToggleSnapTangentAdjustment),
+		FCanExecuteAction(),
+		FIsActionChecked::CreateSP(this, &FSplineComponentVisualizer::IsSnapTangentAdjustment));
 
 	SplineComponentVisualizerActions->MapAction(
 		Commands.SnapAllToSelectedX,
@@ -456,6 +478,16 @@ void FSplineComponentVisualizer::OnRegister()
 		FExecuteAction::CreateSP(this, &FSplineComponentVisualizer::OnSnapSelectedToAxis, EAxis::Z),
 		FCanExecuteAction::CreateSP(this, &FSplineComponentVisualizer::AreMultipleKeysSelected));
 
+	SplineComponentVisualizerActions->MapAction(
+		Commands.StraightenToNext,
+		FExecuteAction::CreateSP(this, &FSplineComponentVisualizer::OnStraightenKey, 1),
+		FCanExecuteAction::CreateSP(this, &FSplineComponentVisualizer::IsKeySelectionValid));
+
+	SplineComponentVisualizerActions->MapAction(
+		Commands.StraightenToPrevious,
+		FExecuteAction::CreateSP(this, &FSplineComponentVisualizer::OnStraightenKey, -1),
+		FCanExecuteAction::CreateSP(this, &FSplineComponentVisualizer::IsKeySelectionValid));
+
 	SplineComponentVisualizerActions->MapAction(
 		Commands.SetLockedAxisNone,
 		FExecuteAction::CreateSP(this, &FSplineComponentVisualizer::OnLockAxis, EAxis::None),
@@ -492,6 +524,12 @@ void FSplineComponentVisualizer::OnRegister()
 		FCanExecuteAction(),
 		FIsActionChecked::CreateSP(this, &FSplineComponentVisualizer::IsDiscontinuousSpline));
 
+	SplineComponentVisualizerActions->MapAction(
+		Commands.ToggleClosedLoop,
+		FExecuteAction::CreateSP(this, &FSplineComponentVisualizer::OnToggleClosedLoop),
+		FCanExecuteAction(),
+		FIsActionChecked::CreateSP(this, &FSplineComponentVisualizer::IsClosedLoop));
+
 	SplineComponentVisualizerActions->MapAction(
 		Commands.ResetToDefault,
 		FExecuteAction::CreateSP(this, &FSplineComponentVisualizer::OnResetToDefault),
@@ -2144,6 +2182,78 @@ void FSplineComponentVisualizer::OnSnapSelectedToAxis(EAxis::Type InAxis)
 	SnapKeysToLastSelectedAxisPosition(InAxis, SnapKeys);
 }
 
+void FSplineComponentVisualizer::OnStraightenKey(int32 Direction)
+{
+	const FScopedTransaction Transaction(LOCTEXT("Straighten To Previous", "Straighten Points Toward Previous"));
+
+	USplineComponent* SplineComp = GetEditedSplineComponent();
+	check(SplineComp != nullptr);
+	check(SelectionState);
+	const TSet<int32>& SelectedKeys = SelectionState->GetSelectedKeys();
+	int32 LastKeyIndexSelected = SelectionState->GetVerifiedLastKeyIndexSelected(SplineComp->GetNumberOfSplinePoints());
+
+	for (int32 CurrentKey : SelectedKeys)
+	{
+		int32 ToKey = CurrentKey + Direction;
+		if (ToKey != INDEX_NONE && ToKey < SplineComp->GetNumberOfSplinePoints())
+		{
+			StraightenKey(CurrentKey, ToKey);
+		}
+	}
+
+	SplineComp->UpdateSpline();
+	SplineComp->bSplineHasBeenEdited = true;
+
+	NotifyPropertyModified(SplineComp, SplineCurvesProperty);
+
+	SelectionState->Modify();
+	SelectionState->SetCachedRotation(SplineComp->GetQuaternionAtSplinePoint(LastKeyIndexSelected, ESplineCoordinateSpace::World));
+
+	GEditor->RedrawLevelEditingViewports(true);
+}
+
+void FSplineComponentVisualizer::StraightenKey(int32 KeyToStraighten, int32 KeyToStraightenToward)
+{
+	USplineComponent* SplineComp = GetEditedSplineComponent();
+	check(SplineComp != nullptr);
+
+	const float TangentLength = SplineComp->GetTangentAtSplinePoint(KeyToStraighten, ESplineCoordinateSpace::Local).Length();
+	FVector StraightenLocation = SplineComp->GetLocationAtSplinePoint(KeyToStraighten, ESplineCoordinateSpace::Local);
+	FVector TowardLocation = SplineComp->GetLocationAtSplinePoint(KeyToStraightenToward, ESplineCoordinateSpace::Local);
+	FVector Direction = TowardLocation - StraightenLocation;
+	Direction.Normalize();
+
+	FVector NewTangent = Direction * TangentLength * (KeyToStraighten > KeyToStraightenToward ? 1 : -1);
+	SplineComp->SetTangentAtSplinePoint(KeyToStraighten, -NewTangent, ESplineCoordinateSpace::Local);
+}
+
+void FSplineComponentVisualizer::OnToggleSnapTangentAdjustment()
+{
+	USplineComponent* SplineComp = GetEditedSplineComponent();
+	check(SplineComp != nullptr);
+
+	SplineComp->Modify();
+	if (AActor* Owner = SplineComp->GetOwner())
+	{
+		Owner->Modify();
+	}
+
+	SplineComp->bAdjustTangentsOnSnap = !SplineComp->bAdjustTangentsOnSnap;
+
+	TArray<FProperty*> Properties;
+	Properties.Add(SplineCurvesProperty);
+	Properties.Add(FindFProperty<FProperty>(USplineComponent::StaticClass(), GET_MEMBER_NAME_CHECKED(USplineComponent, bAdjustTangentsOnSnap)));
+	NotifyPropertiesModified(SplineComp, Properties);
+
+	GEditor->RedrawLevelEditingViewports(true);
+}
+
+bool FSplineComponentVisualizer::IsSnapTangentAdjustment() const
+{
+	USplineComponent* SplineComp = GetEditedSplineComponent();
+	return SplineComp ? SplineComp->bAdjustTangentsOnSnap : false;
+}
+
 void FSplineComponentVisualizer::SnapKeysToLastSelectedAxisPosition(const EAxis::Type InAxis, TArray<int32> InSnapKeys)
 {
 	USplineComponent* SplineComp = GetEditedSplineComponent();
@@ -2190,7 +2300,10 @@ void FSplineComponentVisualizer::SnapKeysToLastSelectedAxisPosition(const EAxis:
 			EditedPosition.OutVal = SplineComp->GetComponentTransform().InverseTransformPosition(NewWorldPos); // convert world-space position to local-space
 
 			// Set point to auto so its tangents will be auto-adjusted after snapping
-			EditedPosition.InterpMode = CIM_CurveAuto;
+			if (SplineComp->bAdjustTangentsOnSnap)
+			{
+				EditedPosition.InterpMode = CIM_CurveAuto;
+			}
 		}
 	}
 
@@ -3159,6 +3272,35 @@ bool FSplineComponentVisualizer::IsDiscontinuousSpline() const
 }
 
 
+void FSplineComponentVisualizer::OnToggleClosedLoop()
+{
+	const FScopedTransaction Transaction(LOCTEXT("ToggleClosedLoop", "Toggle Closed Loop"));
+
+	USplineComponent* SplineComp = GetEditedSplineComponent();
+	check(SplineComp != nullptr);
+
+	SplineComp->Modify();
+	if (AActor* Owner = SplineComp->GetOwner())
+	{
+		Owner->Modify();
+	}
+
+	SplineComp->SetClosedLoop(!SplineComp->IsClosedLoop());
+
+	TArray<FProperty*> Properties;
+	Properties.Add(SplineCurvesProperty);
+	Properties.Add(FindFProperty<FProperty>(USplineComponent::StaticClass(), GET_MEMBER_NAME_CHECKED(USplineComponent, bClosedLoop)));
+	NotifyPropertiesModified(SplineComp, Properties);
+
+	GEditor->RedrawLevelEditingViewports(true);
+}
+
+bool FSplineComponentVisualizer::IsClosedLoop() const
+{
+	USplineComponent* SplineComp = GetEditedSplineComponent();
+	return SplineComp ? SplineComp->IsClosedLoop() : false;
+}
+
 void FSplineComponentVisualizer::OnResetToDefault()
 {
 	const FScopedTransaction Transaction(LOCTEXT("ResetToDefault", "Reset to Default"));
@@ -3473,6 +3615,7 @@ void FSplineComponentVisualizer::GenerateContextMenuSections(FMenuBuilder& InMen
 
 	InMenuBuilder.BeginSection("Spline", LOCTEXT("Spline", "Spline"));
 	{
+		InMenuBuilder.AddMenuEntry(FSplineComponentVisualizerCommands::Get().ToggleClosedLoop);
 		InMenuBuilder.AddMenuEntry(FSplineComponentVisualizerCommands::Get().ResetToDefault);
 	}
 	InMenuBuilder.EndSection();
@@ -3538,6 +3681,11 @@ void FSplineComponentVisualizer::GenerateSnapAlignSubMenu(FMenuBuilder& MenuBuil
 	MenuBuilder.AddMenuEntry(FSplineComponentVisualizerCommands::Get().SnapToLastSelectedX);
 	MenuBuilder.AddMenuEntry(FSplineComponentVisualizerCommands::Get().SnapToLastSelectedY);
 	MenuBuilder.AddMenuEntry(FSplineComponentVisualizerCommands::Get().SnapToLastSelectedZ);
+	MenuBuilder.AddSeparator();
+	MenuBuilder.AddMenuEntry(FSplineComponentVisualizerCommands::Get().StraightenToNext);
+	MenuBuilder.AddMenuEntry(FSplineComponentVisualizerCommands::Get().StraightenToPrevious);
+	MenuBuilder.AddSeparator();
+	MenuBuilder.AddMenuEntry(FSplineComponentVisualizerCommands::Get().ToggleSnapTangentAdjustments);
 }
 
 void FSplineComponentVisualizer::GenerateLockAxisSubMenu(FMenuBuilder& MenuBuilder) const
diff --git a/Engine/Source/Editor/ComponentVisualizers/Public/SplineComponentVisualizer.h b/Engine/Source/Editor/ComponentVisualizers/Public/SplineComponentVisualizer.h
--- a/Engine/Source/Editor/ComponentVisualizers/Public/SplineComponentVisualizer.h
+++ b/Engine/Source/Editor/ComponentVisualizers/Public/SplineComponentVisualizer.h
@@ -364,6 +364,12 @@ protected:
 
 	void OnSnapSelectedToAxis(EAxis::Type InAxis);
 
+	void OnStraightenKey(int32 Direction);
+	void StraightenKey(int32 KeyToStraighten, int32 KeyToStraightenToward);
+
+	void OnToggleSnapTangentAdjustment();
+	bool IsSnapTangentAdjustment() const;
+
 	void OnLockAxis(EAxis::Type InAxis);
 	bool IsLockAxisSet(EAxis::Type InAxis) const; 
 
@@ -379,6 +385,9 @@ protected:
 	void OnSetDiscontinuousSpline();
 	bool IsDiscontinuousSpline() const;
 
+	void OnToggleClosedLoop();
+	bool IsClosedLoop() const;
+
 	void OnResetToDefault();
 	bool CanResetToDefault() const;
 
diff --git a/Engine/Source/Runtime/Engine/Classes/Components/SplineComponent.h b/Engine/Source/Runtime/Engine/Classes/Components/SplineComponent.h
--- a/Engine/Source/Runtime/Engine/Classes/Components/SplineComponent.h
+++ b/Engine/Source/Runtime/Engine/Classes/Components/SplineComponent.h
@@ -283,6 +283,10 @@ public:
 	UPROPERTY(EditAnywhere, Category = Editor)
 	bool bAllowDiscontinuousSpline;
 
+	/** Adjust tangents after snapping. */
+	UPROPERTY(EditAnywhere, Category = Editor)
+	bool bAdjustTangentsOnSnap;
+
 	/** Whether scale visualization should be displayed */
 	UPROPERTY(EditAnywhere, Category = Editor, meta=(InlineEditConditionToggle=true))
 	bool bShouldVisualizeScale;
@@ -887,6 +891,9 @@ private:
 			return DummyPointScale;
 		}
 	}
+
+	// FSplineComponentVisualizer will access some private members when attempting to call NotifyPropertiesModified
+	friend class FSplineComponentVisualizer;
 };
 
 /** Used to store spline data during RerunConstructionScripts */
diff --git a/Engine/Source/Runtime/Engine/Private/Components/SplineComponent.cpp b/Engine/Source/Runtime/Engine/Private/Components/SplineComponent.cpp
--- a/Engine/Source/Runtime/Engine/Private/Components/SplineComponent.cpp
+++ b/Engine/Source/Runtime/Engine/Private/Components/SplineComponent.cpp
@@ -51,6 +51,7 @@ USplineComponent::USplineComponent(const FObjectInitializer& ObjectInitializer)
 	, EditorSelectedSplineSegmentColor(FStyleColors::AccentOrange.GetSpecifiedColor())
 	, EditorTangentColor(FLinearColor(0.718f, 0.589f, 0.921f))
 	, bAllowDiscontinuousSpline(false)
+	, bAdjustTangentsOnSnap(true)
 	, bShouldVisualizeScale(false)
 	, ScaleVisualizationWidth(30.0f)
 #endif
