make spline component visualizer and details customization call the owner actor's posteditmove method to make sure the construction script can re-run as expected

diff --git a/Engine/Source/Editor/ComponentVisualizers/Private/SplineComponentVisualizer.cpp b/Engine/Source/Editor/ComponentVisualizers/Private/SplineComponentVisualizer.cpp
--- a/Engine/Source/Editor/ComponentVisualizers/Private/SplineComponentVisualizer.cpp
+++ b/Engine/Source/Editor/ComponentVisualizers/Private/SplineComponentVisualizer.cpp
@@ -1255,6 +1255,7 @@ bool FSplineComponentVisualizer::HandleInputDelta(FEditorViewportClient* Viewpor
 	ResetTempModes();
 
 	USplineComponent* SplineComp = GetEditedSplineComponent();
+	bool bInputHandled = false;
 	if (SplineComp != nullptr)
 	{
 		if (IsAnySelectedKeyIndexOutOfRange(SplineComp))
@@ -1292,7 +1293,7 @@ bool FSplineComponentVisualizer::HandleInputDelta(FEditorViewportClient* Viewpor
 		if (SelectionState->GetSelectedTangentHandle() != INDEX_NONE)
 		{
 			// Transform the tangent using an EPropertyChangeType::Interactive change. Later on, at the end of mouse tracking, a non-interactive change will be notified via void TrackingStopped :
-			return TransformSelectedTangent(EPropertyChangeType::Interactive, DeltaTranslate);
+			bInputHandled = TransformSelectedTangent(EPropertyChangeType::Interactive, DeltaTranslate);
 		}
 		else if (ViewportClient->IsAltPressed())
 		{
@@ -1333,17 +1334,24 @@ bool FSplineComponentVisualizer::HandleInputDelta(FEditorViewportClient* Viewpor
 					UpdateDuplicateKeyForAltDrag(Drag);
 				}
 
-				return true;
+				bInputHandled = true;
 			}
 		}
 		else
 		{
 			// Transform the spline keys using an EPropertyChangeType::Interactive change. Later on, at the end of mouse tracking, a non-interactive change will be notified via void TrackingStopped :
-			return TransformSelectedKeys(EPropertyChangeType::Interactive, DeltaTranslate, DeltaRotate, DeltaScale);
+			bInputHandled = TransformSelectedKeys(EPropertyChangeType::Interactive, DeltaTranslate, DeltaRotate, DeltaScale);
+		}
+	}
+	if (bInputHandled)
+	{
+		if (AActor* OwnerActor = SplineComp->GetOwner())
+		{
+			OwnerActor->PostEditMove(false);
 		}
 	}
 
-	return false;
+	return bInputHandled;
 }
 
 bool FSplineComponentVisualizer::TransformSelectedTangent(EPropertyChangeType::Type InPropertyChangeType, const FVector& InDeltaTranslate)
@@ -1832,6 +1840,10 @@ bool FSplineComponentVisualizer::HandleSnapTo(const bool bInAlign, const bool bI
 				SplineComp->bSplineHasBeenEdited = true;
 
 				NotifyPropertyModified(SplineComp, SplineCurvesProperty);
+				if (AActor* Owner = SplineComp->GetOwner())
+				{
+					Owner->PostEditMove(true);
+				}
 				
 				if (bInAlign)
 				{
@@ -1856,6 +1868,10 @@ void FSplineComponentVisualizer::TrackingStopped(FEditorViewportClient* InViewpo
 		// After dragging, notify that the spline curves property has changed one last time, this time as a EPropertyChangeType::ValueSet :
 		USplineComponent* SplineComp = GetEditedSplineComponent();
 		NotifyPropertyModified(SplineComp, SplineCurvesProperty, EPropertyChangeType::ValueSet);
+		if (AActor* Owner = SplineComp->GetOwner())
+		{
+			Owner->PostEditMove(true);
+		}
 	}
 }
 
@@ -2118,6 +2134,10 @@ void FSplineComponentVisualizer::SnapKeyToTransform(const ESplineComponentSnapMo
 	SplineComp->bSplineHasBeenEdited = true;
 
 	NotifyPropertyModified(SplineComp, SplineCurvesProperty);
+	if (AActor* Owner = SplineComp->GetOwner())
+	{
+		Owner->PostEditMove(true);
+	}
 
 	if (InSnapMode == ESplineComponentSnapMode::AlignToTangent || InSnapMode == ESplineComponentSnapMode::AlignPerpendicularToTangent)
 	{
@@ -2203,6 +2223,10 @@ void FSplineComponentVisualizer::OnStraightenKey(int32 Direction)
 	SplineComp->bSplineHasBeenEdited = true;
 
 	NotifyPropertyModified(SplineComp, SplineCurvesProperty);
+	if (AActor* OwnerActor = SplineComp->GetOwner())
+	{
+		OwnerActor->PostEditMove(true);
+	}
 
 	SelectionState->Modify();
 	SelectionState->SetCachedRotation(SplineComp->GetQuaternionAtSplinePoint(LastKeyIndexSelected, ESplineCoordinateSpace::World));
@@ -2309,6 +2333,10 @@ void FSplineComponentVisualizer::SnapKeysToLastSelectedAxisPosition(const EAxis:
 	SplineComp->bSplineHasBeenEdited = true;
 
 	NotifyPropertyModified(SplineComp, SplineCurvesProperty);
+	if (AActor* Owner = SplineComp->GetOwner())
+	{
+		Owner->PostEditMove(true);
+	}
 
 	SelectionState->Modify();
 	SelectionState->SetCachedRotation(SplineComp->GetQuaternionAtSplinePoint(LastKeyIndexSelected, ESplineCoordinateSpace::World));
@@ -2437,6 +2465,10 @@ void FSplineComponentVisualizer::OnDuplicateKey()
 	SplineComp->bSplineHasBeenEdited = true;
 
 	NotifyPropertyModified(SplineComp, SplineCurvesProperty);
+	if (AActor* Owner = SplineComp->GetOwner())
+	{
+		Owner->PostEditMove(true);
+	}
 
 	if (NewSelectedKeys.Num() == 1)
 	{
@@ -2733,6 +2765,11 @@ void FSplineComponentVisualizer::SplitSegment(const FVector& InWorldPos, int32 I
 	SplineComp->bSplineHasBeenEdited = true;
 
 	NotifyPropertyModified(SplineComp, SplineCurvesProperty);
+	if (AActor* Owner = SplineComp->GetOwner())
+	{
+		Owner->PostEditMove(true);
+	}
+
 
 	GEditor->RedrawLevelEditingViewports(true);
 }
@@ -3038,6 +3075,10 @@ void FSplineComponentVisualizer::OnDeleteKey()
 	SplineComp->bSplineHasBeenEdited = true;
 
 	NotifyPropertyModified(SplineComp, SplineCurvesProperty);
+	if (AActor* OwnerActor = SplineComp->GetOwner())
+	{
+		OwnerActor->PostEditMove(true);
+	}
 
 	SelectionState->SetCachedRotation(SplineComp->GetQuaternionAtSplinePoint(SelectionState->GetLastKeyIndexSelected(), ESplineCoordinateSpace::World));
 
@@ -3112,6 +3153,10 @@ void FSplineComponentVisualizer::OnResetToAutomaticTangent(EInterpCurveMode Mode
 		SplineComp->bSplineHasBeenEdited = true;
 
 		NotifyPropertyModified(SplineComp, SplineCurvesProperty);
+		if (AActor* OwnerActor = SplineComp->GetOwner())
+		{
+			OwnerActor->PostEditMove(true);
+		}
 
 		SelectionState->Modify();
 		SelectionState->SetCachedRotation(SplineComp->GetQuaternionAtSplinePoint(SelectionState->GetLastKeyIndexSelected(), ESplineCoordinateSpace::World));
@@ -3168,6 +3213,10 @@ void FSplineComponentVisualizer::OnSetKeyType(EInterpCurveMode Mode)
 		SplineComp->bSplineHasBeenEdited = true;
 
 		NotifyPropertyModified(SplineComp, SplineCurvesProperty);
+		if (AActor* OwnerActor = SplineComp->GetOwner())
+		{
+			OwnerActor->PostEditMove(true);
+		}
 
 		SelectionState->Modify();
 		SelectionState->SetCachedRotation(SplineComp->GetQuaternionAtSplinePoint(SelectionState->GetLastKeyIndexSelected(), ESplineCoordinateSpace::World));
@@ -3327,7 +3376,7 @@ void FSplineComponentVisualizer::OnResetToDefault()
 
 	if (AActor* Owner = SplineComp->GetOwner())
 	{
-		Owner->PostEditMove(false);
+		Owner->PostEditMove(true);
 	}
 
 	GEditor->RedrawLevelEditingViewports(true);
diff --git a/Engine/Source/Editor/DetailCustomizations/Private/SplineComponentDetails.cpp b/Engine/Source/Editor/DetailCustomizations/Private/SplineComponentDetails.cpp
--- a/Engine/Source/Editor/DetailCustomizations/Private/SplineComponentDetails.cpp
+++ b/Engine/Source/Editor/DetailCustomizations/Private/SplineComponentDetails.cpp
@@ -954,6 +954,10 @@ void FSplinePointDetails::OnSetInputKey(float NewValue, ETextCommit::Type Commit
 	SplineComp->UpdateSpline();
 	SplineComp->bSplineHasBeenEdited = true;
 	FComponentVisualizer::NotifyPropertyModified(SplineComp, SplineCurvesProperty);
+	if (AActor* Owner = SplineComp->GetOwner())
+	{
+		Owner->PostEditMove(true);
+	}
 	UpdateValues();
 
 	GEditor->RedrawLevelEditingViewports(true);
@@ -1004,6 +1008,10 @@ void FSplinePointDetails::OnSetPosition(float NewValue, ETextCommit::Type Commit
 		SplineComp->UpdateSpline();
 		SplineComp->bSplineHasBeenEdited = true;
 		FComponentVisualizer::NotifyPropertyModified(SplineComp, SplineCurvesProperty, EPropertyChangeType::ValueSet);
+		if (AActor* Owner = SplineComp->GetOwner())
+		{
+			Owner->PostEditMove(true);
+		}
 		UpdateValues();
 	}
 
@@ -1043,6 +1051,10 @@ void FSplinePointDetails::OnSetArriveTangent(float NewValue, ETextCommit::Type C
 		SplineComp->UpdateSpline();
 		SplineComp->bSplineHasBeenEdited = true;
 		FComponentVisualizer::NotifyPropertyModified(SplineComp, SplineCurvesProperty, EPropertyChangeType::ValueSet);
+		if (AActor* Owner = SplineComp->GetOwner())
+		{
+			Owner->PostEditMove(true);
+		}
 		UpdateValues();
 	}
 
@@ -1082,6 +1094,10 @@ void FSplinePointDetails::OnSetLeaveTangent(float NewValue, ETextCommit::Type Co
 		SplineComp->UpdateSpline();
 		SplineComp->bSplineHasBeenEdited = true;
 		FComponentVisualizer::NotifyPropertyModified(SplineComp, SplineCurvesProperty, EPropertyChangeType::ValueSet);
+		if (AActor* Owner = SplineComp->GetOwner())
+		{
+			Owner->PostEditMove(true);
+		}
 		UpdateValues();
 	}
 
@@ -1160,6 +1176,10 @@ void FSplinePointDetails::OnSetRotation(float NewValue, ETextCommit::Type Commit
 		SplineComp->UpdateSpline();
 		SplineComp->bSplineHasBeenEdited = true;
 		FComponentVisualizer::NotifyPropertyModified(SplineComp, SplineCurvesProperty, EPropertyChangeType::ValueSet);
+		if (AActor* Owner = SplineComp->GetOwner())
+		{
+			Owner->PostEditMove(true);
+		}
 		UpdateValues();
 	}
 	GEditor->RedrawLevelEditingViewports(true);
@@ -1197,6 +1217,10 @@ void FSplinePointDetails::OnSetScale(float NewValue, ETextCommit::Type CommitInf
 		SplineComp->UpdateSpline();
 		SplineComp->bSplineHasBeenEdited = true;
 		FComponentVisualizer::NotifyPropertyModified(SplineComp, SplineCurvesProperty, EPropertyChangeType::ValueSet);
+		if (AActor* Owner = SplineComp->GetOwner())
+		{
+			Owner->PostEditMove(true);
+		}
 		UpdateValues();
 	}
 
@@ -1253,6 +1277,10 @@ void FSplinePointDetails::OnSplinePointTypeChanged(TSharedPtr<FString> NewValue,
 	SplineComp->UpdateSpline();
 	SplineComp->bSplineHasBeenEdited = true;
 	FComponentVisualizer::NotifyPropertyModified(SplineComp, SplineCurvesProperty);
+	if (AActor* Owner = SplineComp->GetOwner())
+	{
+		Owner->PostEditMove(true);
+	}
 	UpdateValues();
 
 	GEditor->RedrawLevelEditingViewports(true);
