Fixed linking error.

diff --git a/Engine/Source/Editor/ComponentVisualizers/Private/RectLightComponentVisualizer.cpp b/Engine/Source/Editor/ComponentVisualizers/Private/RectLightComponentVisualizer.cpp
--- a/Engine/Source/Editor/ComponentVisualizers/Private/RectLightComponentVisualizer.cpp
+++ b/Engine/Source/Editor/ComponentVisualizers/Private/RectLightComponentVisualizer.cpp
@@ -16,7 +16,7 @@
 #include "Templates/Casts.h"
 
 int32 GVisualizeCullingBarnDoors = 0;
-FAutoConsoleVariableRef CVarLightGridHZBCull(
+static FAutoConsoleVariableRef CVarVisualizeCullingBarnDoors(
 	TEXT("r.RectLight.VisualizeCullingBarnDoors"),
 	GVisualizeCullingBarnDoors,
 	TEXT("Whether to render a visualization of the barn doors used to cull the rect light."),
