[Truncation Warnings] Fix warnings in component visualizers module

#rnx

diff --git a/Engine/Source/Editor/ComponentVisualizers/Private/PointLightComponentVisualizer.cpp b/Engine/Source/Editor/ComponentVisualizers/Private/PointLightComponentVisualizer.cpp
--- a/Engine/Source/Editor/ComponentVisualizers/Private/PointLightComponentVisualizer.cpp
+++ b/Engine/Source/Editor/ComponentVisualizers/Private/PointLightComponentVisualizer.cpp
@@ -291,7 +291,7 @@ bool FTextureLightProfileVisualizer::UpdateIntensitiesCache(UTextureLightProfile
 	TArray64< uint8 > MipData;
 	verify( TextureLightProfile->Source.GetMipData( MipData, 0 ) );
 
-	FTextureLightProfileData TextureLightProfileData{ MipData, TextureLightProfile->Source.GetSizeX(), TextureLightProfile->Source.GetSizeY(), TextureLightProfile->Source.GetBytesPerPixel() };
+	FTextureLightProfileData TextureLightProfileData{ MipData, (int32)TextureLightProfile->Source.GetSizeX(), (int32)TextureLightProfile->Source.GetSizeY(), (int32)TextureLightProfile->Source.GetBytesPerPixel() };
 
 	int32 NumIntensities = ( FPolarCoordinates::MaxAltitude / FPolarSampler::Step ) * ( FPolarCoordinates::MaxAzimuth / FPolarSampler::Step );
 
diff --git a/Engine/Source/Editor/ComponentVisualizers/Private/SplineComponentVisualizer.cpp b/Engine/Source/Editor/ComponentVisualizers/Private/SplineComponentVisualizer.cpp
--- a/Engine/Source/Editor/ComponentVisualizers/Private/SplineComponentVisualizer.cpp
+++ b/Engine/Source/Editor/ComponentVisualizers/Private/SplineComponentVisualizer.cpp
@@ -2239,7 +2239,7 @@ void FSplineComponentVisualizer::StraightenKey(int32 KeyToStraighten, int32 KeyT
 	USplineComponent* SplineComp = GetEditedSplineComponent();
 	check(SplineComp != nullptr);
 
-	const float TangentLength = SplineComp->GetTangentAtSplinePoint(KeyToStraighten, ESplineCoordinateSpace::Local).Length();
+	const FVector::FReal TangentLength = SplineComp->GetTangentAtSplinePoint(KeyToStraighten, ESplineCoordinateSpace::Local).Length();
 	FVector StraightenLocation = SplineComp->GetLocationAtSplinePoint(KeyToStraighten, ESplineCoordinateSpace::Local);
 	FVector TowardLocation = SplineComp->GetLocationAtSplinePoint(KeyToStraightenToward, ESplineCoordinateSpace::Local);
 	FVector Direction = TowardLocation - StraightenLocation;
