[Truncation Warnings] Enable truncation warnings in ComponentVisualizers module

#rnx

diff --git a/Engine/Source/Editor/ComponentVisualizers/ComponentVisualizers.Build.cs b/Engine/Source/Editor/ComponentVisualizers/ComponentVisualizers.Build.cs
--- a/Engine/Source/Editor/ComponentVisualizers/ComponentVisualizers.Build.cs
+++ b/Engine/Source/Editor/ComponentVisualizers/ComponentVisualizers.Build.cs
@@ -6,6 +6,9 @@ public class ComponentVisualizers : ModuleRules
 {
 	public ComponentVisualizers(ReadOnlyTargetRules Target) : base(Target)
 	{
+		// Enable truncation warnings in this module
+		CppCompileWarningSettings.UnsafeTypeCastWarningLevel = WarningLevel.Warning;
+
 		PrivateDependencyModuleNames.AddRange(
 			new string[] {
 				"Core",
diff --git a/Engine/Source/Editor/ComponentVisualizers/Private/PointLightComponentVisualizer.cpp b/Engine/Source/Editor/ComponentVisualizers/Private/PointLightComponentVisualizer.cpp
--- a/Engine/Source/Editor/ComponentVisualizers/Private/PointLightComponentVisualizer.cpp
+++ b/Engine/Source/Editor/ComponentVisualizers/Private/PointLightComponentVisualizer.cpp
@@ -67,8 +67,8 @@ namespace TextureLightProfileVisualizerImpl
 		const int32 BytesPerPixel = TextureLightProfileData.BytesPerPixel;
 
 		// not 100% like GPU hardware but simple and almost the same
-		float UnNormalizedX = FMath::Clamp(X * SizeX, 0.0f, (float)(SizeX - 1));
-		float UnNormalizedY = FMath::Clamp(Y * SizeY, 0.0f, (float)(SizeY - 1));
+		float UnNormalizedX = FMath::Clamp(X * static_cast<float>(SizeX), 0.0f, static_cast<float>(SizeX - 1));
+		float UnNormalizedY = FMath::Clamp(Y * static_cast<float>(SizeY), 0.0f,static_cast<float>(SizeY - 1));
 
 		int32 X0 = (uint32)UnNormalizedX;
 		int32 X1 = FMath::Min(X0 + 1, SizeX - 1);
@@ -76,8 +76,8 @@ namespace TextureLightProfileVisualizerImpl
 		int32 Y0 = (uint32)UnNormalizedY;
 		int32 Y1 = FMath::Min(Y0 + 1, SizeY - 1);
 
-		float XFraction = UnNormalizedX - X0;
-		float YFraction = UnNormalizedY - Y0;
+		float XFraction = UnNormalizedX - static_cast<float>(X0);
+		float YFraction = UnNormalizedY - static_cast<float>(Y0);
 
 		float V00 = reinterpret_cast<const FFloat16*>(&TextureLightProfileData.Data[ (Y0 * SizeX * BytesPerPixel) + (X0 * BytesPerPixel) ])->GetFloat();
 		float V10 = reinterpret_cast<const FFloat16*>(&TextureLightProfileData.Data[ (Y1 * SizeX * BytesPerPixel) + (X0 * BytesPerPixel) ])->GetFloat();
diff --git a/Engine/Source/Editor/ComponentVisualizers/Private/SplineComponentVisualizer.cpp b/Engine/Source/Editor/ComponentVisualizers/Private/SplineComponentVisualizer.cpp
--- a/Engine/Source/Editor/ComponentVisualizers/Private/SplineComponentVisualizer.cpp
+++ b/Engine/Source/Editor/ComponentVisualizers/Private/SplineComponentVisualizer.cpp
@@ -709,7 +709,7 @@ void FSplineComponentVisualizer::DrawVisualization(const UActorComponent* Compon
 					{
 						float Sin;
 						float Cos;
-						FMath::SinCos(&Sin, &Cos, ArcIndex * PI / ArcPoints);
+						FMath::SinCos(&Sin, &Cos, static_cast<float>(ArcIndex) * PI / static_cast<float>(ArcPoints));
 						const FVector NewArcPos = NewKeyPos + Cos * NewKeyRightVector * NewKeyScale.Y + Sin * NewKeyUpVector * NewKeyScale.Z;
 						PDI->DrawLine(OldArcPos, NewArcPos, KeyColor, SDPG_Foreground);
 						OldArcPos = NewArcPos;
@@ -767,8 +767,8 @@ void FSplineComponentVisualizer::DrawVisualization(const UActorComponent* Compon
 
 					for (int32 StepIdx = 1; StepIdx <= NumSteps; StepIdx++)
 					{
-						const float StepRatio = StepIdx / static_cast<float>(NumSteps);
-						const float Key = EndIdx + StepRatio;
+						const float StepRatio = static_cast<float>(StepIdx) / static_cast<float>(NumSteps);
+						const float Key = static_cast<float>(EndIdx) + StepRatio;
 						const FVector NewPos = SplineComp->GetLocationAtSplineInputKey(Key, ESplineCoordinateSpace::World);
 						const FVector NewRightVector = SplineComp->GetRightVectorAtSplineInputKey(Key, ESplineCoordinateSpace::World);
 						const FVector NewScale = SplineComp->GetScaleAtSplineInputKey(Key) * DefaultScale;
@@ -845,8 +845,8 @@ void FSplineComponentVisualizer::DrawVisualizationHUD(const UActorComponent* Com
 					int32 XL;
 					int32 YL;
 					StringSize(GEngine->GetLargeFont(), XL, YL, *SnapHelpText.ToString());
-					const float DrawPositionX = FMath::FloorToFloat(CanvasRect.Min.X + (CanvasRect.Width() - XL) * 0.5f);
-					const float DrawPositionY = CanvasRect.Min.Y + 50.0f;
+					const float DrawPositionX = FMath::FloorToFloat(static_cast<float>(CanvasRect.Min.X) + static_cast<float>(CanvasRect.Width() - XL) * 0.5f);
+					const float DrawPositionY = static_cast<float>(CanvasRect.Min.Y) + 50.0f;
 					Canvas->DrawShadowedString(DrawPositionX, DrawPositionY, *SnapHelpText.ToString(), GEngine->GetLargeFont(), FLinearColor::Yellow);
 				};
 
@@ -1043,7 +1043,7 @@ bool FSplineComponentVisualizer::VisProxyHandleClick(FEditorViewportClient* InVi
 
 				for (int32 Step = 1; Step < NumSubdivisions; Step++)
 				{
-					const float SubsegmentEndKey = SelectedSegmentIndex + Step / static_cast<float>(NumSubdivisions);
+					const float SubsegmentEndKey = static_cast<float>(SelectedSegmentIndex) + static_cast<float>(Step) / static_cast<float>(NumSubdivisions);
 					const FVector SubsegmentEnd = SplineComp->GetLocationAtSplineInputKey(SubsegmentEndKey, ESplineCoordinateSpace::World);
 
 					FVector SplineClosest;
diff --git a/Engine/Source/Editor/ComponentVisualizers/Private/SplineMeshComponentVisualizer.cpp b/Engine/Source/Editor/ComponentVisualizers/Private/SplineMeshComponentVisualizer.cpp
--- a/Engine/Source/Editor/ComponentVisualizers/Private/SplineMeshComponentVisualizer.cpp
+++ b/Engine/Source/Editor/ComponentVisualizers/Private/SplineMeshComponentVisualizer.cpp
@@ -100,7 +100,7 @@ void FSplineMeshComponentVisualizer::DrawVisualization(const UActorComponent* Co
 		FVector StartPos = SplineMeshComp->GetComponentTransform().TransformPosition(Spline.Points[0].OutVal);
 		for (int32 Step = 1; Step < 32; Step++)
 		{
-			const FVector EndPos = SplineMeshComp->GetComponentTransform().TransformPosition(Spline.Eval(Step / 32.0f, FVector::ZeroVector));
+			const FVector EndPos = SplineMeshComp->GetComponentTransform().TransformPosition(Spline.Eval(static_cast<float>(Step) / 32.0f, FVector::ZeroVector));
 			PDI->DrawLine(StartPos, EndPos, Color, SDPG_Foreground);
 			StartPos = EndPos;
 		}
