SplineComponentVisualizer: Correct some regressions caused by recent changes to SplineComponent.

This corrects 2 separate regressions where the tangent type of a newly created point would be wrongly set to user tangents.

This also corrects a 3rd regression where dragging after creating a new point would affect the wrong point (index off by 1 error).

diff --git a/Engine/Source/Editor/ComponentVisualizers/Private/SplineComponentVisualizer.cpp b/Engine/Source/Editor/ComponentVisualizers/Private/SplineComponentVisualizer.cpp
--- a/Engine/Source/Editor/ComponentVisualizers/Private/SplineComponentVisualizer.cpp
+++ b/Engine/Source/Editor/ComponentVisualizers/Private/SplineComponentVisualizer.cpp
@@ -2639,11 +2639,11 @@ void FSplineComponentVisualizer::SplitSegment(const FVector& InWorldPos, int32 I
 
 	// Determine which index to use when copying interp mode
 	int32 SourceIndex = bCopyFromSegmentBeginIndex ? SegmentBeginIndex : SegmentEndIndex;
-
+	ESplinePointType::Type SplinePointType = SplineComp->GetSplinePointType(SourceIndex);
 	SplineComp->AddSplinePointAtIndex(LocalPos, SegmentSplitIndex, ESplineCoordinateSpace::Local, false);
 	SplineComp->SetQuaternionAtSplinePoint(SegmentSplitIndex, NewRot, ESplineCoordinateSpace::Local, false);
 	SplineComp->SetScaleAtSplinePoint(SegmentSplitIndex, NewScale, false);
-	SplineComp->SetSplinePointType(SegmentSplitIndex, SplineComp->GetSplinePointType(SourceIndex), false);
+	SplineComp->SetSplinePointType(SegmentSplitIndex, SplinePointType, false);
 
 	// Set selection to new key
 	ChangeSelectionState(SegmentSplitIndex, false);
@@ -2723,7 +2723,7 @@ void FSplineComponentVisualizer::UpdateSplitSegment(const FVector& InDrag)
 	DuplicateCacheSplitSegmentParam = t;
 
 	// Update location
-	SplineComp->SetLocationAtSplinePoint(SegmentStartIndex, LocalPos, ESplineCoordinateSpace::Local, false);
+	SplineComp->SetLocationAtSplinePoint(SegmentSplitIndex, LocalPos, ESplineCoordinateSpace::Local, false);
 
 	// Update scale
 	const FVector PrevScale = SplineComp->GetScaleAtSplinePoint(SegmentStartIndex);
@@ -2731,9 +2731,11 @@ void FSplineComponentVisualizer::UpdateSplitSegment(const FVector& InDrag)
 	SplineComp->SetScaleAtSplinePoint(SegmentSplitIndex, FMath::LerpStable(PrevScale, NextScale, t), false);
 
 	// Update rot
+	ESplinePointType::Type SplinePointType = SplineComp->GetSplinePointType(SegmentSplitIndex);
 	FQuat PrevRot = SplineComp->GetRotationAtSplinePoint(SegmentStartIndex, ESplineCoordinateSpace::Local).Quaternion();
 	FQuat NextRot = SplineComp->GetRotationAtSplinePoint(SegmentEndIndex, ESplineCoordinateSpace::Local).Quaternion();
 	SplineComp->SetRotationAtSplinePoint(SegmentSplitIndex, FMath::Lerp(PrevRot, NextRot, t).Rotator(), ESplineCoordinateSpace::Local, false);
+	SplineComp->SetSplinePointType(SegmentSplitIndex, SplinePointType, false);
 
 	// Update metadata
 	if (USplineMetadata* SplineMetadata = SplineComp->GetSplinePointsMetadata())
