SplineComponentVisualizer: Bring back exponential scaling.

Spline point scaling was recently made linear to be more consistent with the rest of the editor, but this change should have been made more intentionally and separately from a more general refactor of spline component. This should arguably be linear, but this decision can be made later.

The change to water spline component is just bringing back previous behavior which existed to handle exponential scaling.

diff --git a/Engine/Plugins/Experimental/Water/Source/Runtime/Private/WaterSplineComponent.cpp b/Engine/Plugins/Experimental/Water/Source/Runtime/Private/WaterSplineComponent.cpp
--- a/Engine/Plugins/Experimental/Water/Source/Runtime/Private/WaterSplineComponent.cpp
+++ b/Engine/Plugins/Experimental/Water/Source/Runtime/Private/WaterSplineComponent.cpp
@@ -244,6 +244,22 @@ PRAGMA_ENABLE_DEPRECATION_WARNINGS
 			}
 
 			SetScaleAtSplinePoint(Point, Scale, false);
+
+			// #hack: temporarily clamp the tangents to a sensible range to prevent multiplicatively scaling until they hit infinity and cause a crash
+			auto ClampVec3 = [](const FVector& Vec3, double Min, double Max) -> FVector {
+				FVector Result;
+				Result.X = FMath::Clamp(Vec3.X, Min, Max);
+				Result.Y = FMath::Clamp(Vec3.Y, Min, Max);
+				Result.Z = FMath::Clamp(Vec3.Z, Min, Max);
+				return Result;
+			};
+			FVector ArriveTangent = GetArriveTangentAtSplinePoint(Point, ESplineCoordinateSpace::Local);
+			FVector LeaveTangent = GetLeaveTangentAtSplinePoint(Point, ESplineCoordinateSpace::Local);
+			
+			constexpr double MaxTangentValue = 1.e10L;
+			ArriveTangent = ClampVec3(ArriveTangent, -MaxTangentValue, MaxTangentValue);
+			LeaveTangent = ClampVec3(LeaveTangent, -MaxTangentValue, MaxTangentValue);
+			SetTangentsAtSplinePoint(Point, ArriveTangent, LeaveTangent, ESplineCoordinateSpace::Local, false);
 		}
 	}
 
diff --git a/Engine/Source/Editor/ComponentVisualizers/Private/SplineComponentVisualizer.cpp b/Engine/Source/Editor/ComponentVisualizers/Private/SplineComponentVisualizer.cpp
--- a/Engine/Source/Editor/ComponentVisualizers/Private/SplineComponentVisualizer.cpp
+++ b/Engine/Source/Editor/ComponentVisualizers/Private/SplineComponentVisualizer.cpp
@@ -1456,7 +1456,7 @@ bool FSplineComponentVisualizer::TransformSelectedKeys(EPropertyChangeType::Type
 
 			if (InDeltaScale.X != 0.0f)
 			{
-				const FVector NewTangent = SplineComp->GetLeaveTangentAtSplinePoint(SelectedKeyIndex, ESplineCoordinateSpace::Local) + InDeltaScale.X;
+				const FVector NewTangent = SplineComp->GetLeaveTangentAtSplinePoint(SelectedKeyIndex, ESplineCoordinateSpace::Local) * (1.f + InDeltaScale.X);
 				SplineComp->SetTangentsAtSplinePoint(SelectedKeyIndex, NewTangent, NewTangent, ESplineCoordinateSpace::Local, false);
 			}
 
@@ -1464,7 +1464,7 @@ bool FSplineComponentVisualizer::TransformSelectedKeys(EPropertyChangeType::Type
 			{
 				// Scale in Y adjusts the scale spline
 				FVector NewScale = SplineComp->GetScaleAtSplinePoint(SelectedKeyIndex);
-				NewScale.Y += InDeltaScale.Y;
+				NewScale.Y *= (1.f + InDeltaScale.Y);
 				SplineComp->SetScaleAtSplinePoint(SelectedKeyIndex, NewScale, false);
 			}
 
@@ -1472,7 +1472,7 @@ bool FSplineComponentVisualizer::TransformSelectedKeys(EPropertyChangeType::Type
 			{
 				// Scale in Z adjusts the scale spline
 				FVector NewScale = SplineComp->GetScaleAtSplinePoint(SelectedKeyIndex);
-				NewScale.Z += InDeltaScale.Z;
+				NewScale.Z *= (1.f + InDeltaScale.Z);
 				SplineComp->SetScaleAtSplinePoint(SelectedKeyIndex, NewScale, false);
 			}
 		}
