SplineComponentVisualizer: Disable attribute editing menu when the underlying spline does not support attributes

diff --git a/Engine/Source/Editor/ComponentVisualizers/Private/SplineComponentVisualizer.cpp b/Engine/Source/Editor/ComponentVisualizers/Private/SplineComponentVisualizer.cpp
--- a/Engine/Source/Editor/ComponentVisualizers/Private/SplineComponentVisualizer.cpp
+++ b/Engine/Source/Editor/ComponentVisualizers/Private/SplineComponentVisualizer.cpp
@@ -3690,8 +3690,11 @@ void FSplineComponentVisualizer::GenerateContextMenuSections(FMenuBuilder& InMen
 		InMenuBuilder.AddMenuEntry(FSplineComponentVisualizerCommands::Get().DiscontinuousSpline);
 	}
 	InMenuBuilder.EndSection();
-	
-	GenerateAttributeMenuSection(InMenuBuilder);
+
+	if (SplineComp && SplineComp->SupportsAttributes())
+	{
+		GenerateAttributeMenuSection(InMenuBuilder);
+	}
 }
 
 void FSplineComponentVisualizer::GenerateSelectSplinePointsSubMenu(FMenuBuilder& MenuBuilder) const
