PR #12674: Fix SplineComponent Losing Custom Tangents of adjacent points when inserting point.
-Needed updates to fit with teh current implementation

diff --git a/Engine/Source/Editor/ComponentVisualizers/Private/SplineComponentVisualizer.cpp b/Engine/Source/Editor/ComponentVisualizers/Private/SplineComponentVisualizer.cpp
--- a/Engine/Source/Editor/ComponentVisualizers/Private/SplineComponentVisualizer.cpp
+++ b/Engine/Source/Editor/ComponentVisualizers/Private/SplineComponentVisualizer.cpp
@@ -2750,16 +2750,6 @@ void FSplineComponentVisualizer::SplitSegment(const FVector& InWorldPos, int32 I
 	{
 		SegmentEndIndex = 0;
 	}
-	
-	// Set adjacent points to CurveAuto so their tangents adjust automatically as new point moves.
-	if (SplineComp->GetSplinePointType(SegmentBeginIndex) == ConvertInterpCurveModeToSplinePointType(CIM_CurveUser))
-	{
-		SplineComp->SetSplinePointType(SegmentBeginIndex, ConvertInterpCurveModeToSplinePointType(CIM_CurveAuto), false);
-	}
-	if (SplineComp->GetSplinePointType(SegmentEndIndex) == ConvertInterpCurveModeToSplinePointType(CIM_CurveUser))
-	{
-		SplineComp->SetSplinePointType(SegmentEndIndex, ConvertInterpCurveModeToSplinePointType(CIM_CurveAuto), false);
-	}
 
 	// Compute interpolated scale
 	FVector NewScale;
@@ -2775,11 +2765,13 @@ void FSplineComponentVisualizer::SplitSegment(const FVector& InWorldPos, int32 I
 
 	// Determine which index to use when copying interp mode
 	int32 SourceIndex = bCopyFromSegmentBeginIndex ? SegmentBeginIndex : SegmentEndIndex;
-	ESplinePointType::Type SplinePointType = SplineComp->GetSplinePointType(SourceIndex);
+	ESplinePointType::Type SourceSplinePointType = SplineComp->GetSplinePointType(SourceIndex);
+	// If the spline interpolation mode of the source point is a custom tangent curve, change it to be an auto curve
+	ESplinePointType::Type NewSplinePointType = SourceSplinePointType == ESplinePointType::CurveCustomTangent ? ESplinePointType::Curve : SourceSplinePointType;
 	SplineComp->AddSplinePointAtIndex(LocalPos, SegmentSplitIndex, ESplineCoordinateSpace::Local, false);
 	SplineComp->SetQuaternionAtSplinePoint(SegmentSplitIndex, NewRot, ESplineCoordinateSpace::Local, false);
 	SplineComp->SetScaleAtSplinePoint(SegmentSplitIndex, NewScale, false);
-	SplineComp->SetSplinePointType(SegmentSplitIndex, SplinePointType, false);
+	SplineComp->SetSplinePointType(SegmentSplitIndex, NewSplinePointType, false);
 
 	// Set selection to new key
 	ChangeSelectionState(SegmentSplitIndex, false);
