SplineCore: Several related bug fixes:

1) Reintroduce ISplineInterface::IsEqual, which is used for checking equality of two ISplineInterfaces without knowing what type of spline they are
2) Fix TMultiSpline::operator==, which was not properly checking for equality of child splines
3) Ensure we mark the spline component as modified when interacting with attributes to ensure changes are saved and undo/redo works as intended.

diff --git a/Engine/Source/Editor/ComponentVisualizers/Private/SplineComponentVisualizer.cpp b/Engine/Source/Editor/ComponentVisualizers/Private/SplineComponentVisualizer.cpp
--- a/Engine/Source/Editor/ComponentVisualizers/Private/SplineComponentVisualizer.cpp
+++ b/Engine/Source/Editor/ComponentVisualizers/Private/SplineComponentVisualizer.cpp
@@ -1594,6 +1594,10 @@ bool FSplineComponentVisualizer::TransformSelectedAttribute(EPropertyChangeType:
 		check(SelectedAttributeName != NAME_None);
 
 		SplineComp->Modify();
+		if (AActor* Owner = SplineComp->GetOwner())
+		{
+			Owner->Modify();
+		}
 		
 		const float CurrentParam = SplineComp->GetFloatPropertyInputKeyAtIndex(SelectedAttributeIndex, SelectedAttributeName);
 		const FVector CurrentPosition = SplineComp->GetLocationAtSplineInputKey(CurrentParam, ESplineCoordinateSpace::World);
@@ -3339,7 +3343,7 @@ bool FSplineComponentVisualizer::CanResetToDefault() const
 
 void FSplineComponentVisualizer::OnAddAttributeKey()
 {
-	const FScopedTransaction Transaction(LOCTEXT("AddSplinePoint", "Add Spline Point"));
+	const FScopedTransaction Transaction(LOCTEXT("AddAttributePoint", "Add Attribute Point"));
 
 	ResetTempModes();
 
@@ -3349,6 +3353,12 @@ void FSplineComponentVisualizer::OnAddAttributeKey()
 
 	SelectionState->Modify();
 
+	SplineComp->Modify();
+	if (AActor* Owner = SplineComp->GetOwner())
+	{
+		Owner->Modify();
+	}
+	
 	const float NearestParam = SplineComp->FindInputKeyClosestToWorldLocation(SelectionState->GetSelectedSplinePosition());
 	// const float 
 	int32 NewAttributeIndex = SplineComp->SetFloatPropertyAtSplineInputKey(NearestParam, 0.f, SelectionState->GetSelectedAttributeName());
@@ -3358,6 +3368,9 @@ void FSplineComponentVisualizer::OnAddAttributeKey()
 	SelectionState->SetSelectedAttribute(NewAttributeIndex, SelectionState->GetSelectedAttributeName());
 	SelectionState->SetSelectedSplinePosition(FVector::ZeroVector);
 	SelectionState->SetCachedRotation(SplineComp->GetQuaternionAtSplineInputKey(NearestParam, ESplineCoordinateSpace::World));
+
+	SplineComp->bSplineHasBeenEdited = true;
+	NotifyPropertyModified(SplineComp, SplineCurvesProperty, EPropertyChangeType::ArrayAdd);
 }
 
 bool FSplineComponentVisualizer::CanAddAttributeKey() const
@@ -3369,7 +3382,7 @@ bool FSplineComponentVisualizer::CanAddAttributeKey() const
 
 void FSplineComponentVisualizer::OnDeleteAttributeKey()
 {
-	const FScopedTransaction Transaction(LOCTEXT("RemoveSplinePoint", "Remove Spline Point"));
+	const FScopedTransaction Transaction(LOCTEXT("RemoveAttributePoint", "Remove Attribute Point"));
 
 	ResetTempModes();
 
@@ -3379,6 +3392,12 @@ void FSplineComponentVisualizer::OnDeleteAttributeKey()
 
 	SelectionState->Modify();
 
+	SplineComp->Modify();
+	if (AActor* Owner = SplineComp->GetOwner())
+	{
+		Owner->Modify();
+	}
+
 	SplineComp->RemovePropertyAtIndex(SelectionState->GetSelectedAttributeIndex(), SelectionState->GetSelectedAttributeName());
 
 	// Update selection
@@ -3396,6 +3415,9 @@ void FSplineComponentVisualizer::OnDeleteAttributeKey()
 		SelectionState->SetCachedRotation(SplineComp->GetQuaternionAtSplineInputKey(NewParam, ESplineCoordinateSpace::World));
 	}
 
+	SplineComp->bSplineHasBeenEdited = true;
+	NotifyPropertyModified(SplineComp, SplineCurvesProperty, EPropertyChangeType::ArrayRemove);
+	
 	// Necessary because the attribute proxies store invalid indices, we must redraw them.
 	GEditor->RedrawAllViewports();
 }
@@ -3835,14 +3857,23 @@ TSharedPtr<SWidget> FSplineComponentVisualizer::GenerateAttributeEditorWidget()
 	
 				if (AttributeIndex != INDEX_NONE && AttributeName != NAME_None)
 				{
-					const float NewParam = FMath::Clamp(NewParameter, LowerBound, UpperBound);
-					const int32 NewIndex = LocalSplineComp->SetFloatPropertyInputKeyAtIndex(AttributeIndex, NewParam, AttributeName);
+					const FScopedTransaction Transaction(LOCTEXT("SetAttributeParameter", "Set Attribute Parameter"));
 
-					// Necessary if SetFloatPropertyInputKeyAtIndex caused attribute storage reordering.
-					SelectionState->SetSelectedAttribute(NewIndex, SelectionState->GetSelectedAttributeName());
+					LocalSplineComp->Modify();
+					if (AActor* Owner = LocalSplineComp->GetOwner())
+					{
+						Owner->Modify();
+					}
 					
-					// Without this, the viewport will not reflect changes until the context menu is closed.
+					const float NewParam = FMath::Clamp(NewParameter, LowerBound, UpperBound);
 					SelectionState->SetCachedRotation(LocalSplineComp->GetQuaternionAtSplineInputKey(NewParam, ESplineCoordinateSpace::World));
+					if (const int32 NewIndex = LocalSplineComp->SetFloatPropertyInputKeyAtIndex(AttributeIndex, NewParam, AttributeName); NewIndex != AttributeIndex)
+					{
+						SelectionState->SetSelectedAttribute(NewIndex, SelectionState->GetSelectedAttributeName());
+					}
+
+					LocalSplineComp->bSplineHasBeenEdited = true;
+					NotifyPropertyModified(LocalSplineComp, SplineCurvesProperty, EPropertyChangeType::ValueSet);
 					GEditor->RedrawAllViewports();
 				}
 			}
@@ -3874,8 +3905,19 @@ TSharedPtr<SWidget> FSplineComponentVisualizer::GenerateAttributeEditorWidget()
 	
 			if (AttributeIndex != INDEX_NONE && AttributeName != NAME_None)
 			{
+				const FScopedTransaction Transaction(LOCTEXT("SetAttributeValue", "Set Attribute Value"));
+
+				LocalSplineComp->Modify();
+				if (AActor* Owner = LocalSplineComp->GetOwner())
+				{
+					Owner->Modify();
+				}
+				
 				LocalSplineComp->SetFloatPropertyAtIndex(AttributeIndex, NewValue, AttributeName);
-				GEditor->RedrawAllViewports();	// Force redraw while context menu is preventing it
+				
+				LocalSplineComp->bSplineHasBeenEdited = true;
+				NotifyPropertyModified(LocalSplineComp, SplineCurvesProperty, EPropertyChangeType::ValueSet);
+				GEditor->RedrawAllViewports();
 			}
 		}
 	};
diff --git a/Engine/Source/Runtime/GeometryCore/Public/Splines/BSpline.h b/Engine/Source/Runtime/GeometryCore/Public/Splines/BSpline.h
--- a/Engine/Source/Runtime/GeometryCore/Public/Splines/BSpline.h
+++ b/Engine/Source/Runtime/GeometryCore/Public/Splines/BSpline.h
@@ -68,6 +68,18 @@ public:
 		Values.Empty();
 		Knots.Empty();
 	}
+
+	virtual bool IsEqual(const ISplineInterface* OtherSpline) const override
+	{
+		if (OtherSpline->GetTypeId() == GetTypeId())
+		{
+			const TBSpline* Other = static_cast<const TBSpline*>(OtherSpline);
+			return operator==(*Other);
+		}
+		
+		return false;
+	}
+	
 	virtual bool Serialize(FArchive& Ar) override
 	{
 		// Call immediate parent's Serialize
@@ -115,7 +127,6 @@ public:
     
 		return Clone;
 	}
-	// virtual bool IsEqual(const ISplineInterface& Other) const override { /*todo*/ }
 
 	// TSplineInterface Implementation
 	virtual ValueType EvaluateImpl(float Parameter) const override
diff --git a/Engine/Source/Runtime/GeometryCore/Public/Splines/MultiSpline.h b/Engine/Source/Runtime/GeometryCore/Public/Splines/MultiSpline.h
--- a/Engine/Source/Runtime/GeometryCore/Public/Splines/MultiSpline.h
+++ b/Engine/Source/Runtime/GeometryCore/Public/Splines/MultiSpline.h
@@ -20,7 +20,9 @@ namespace Spline
  * Provides attribute functionality while allowing direct access to implementation-specific methods
  */
 template <typename SPLINETYPE>
-class TMultiSpline : public TSplineWrapper<SPLINETYPE>
+class TMultiSpline :
+	public TSplineWrapper<SPLINETYPE>,
+	private TSelfRegisteringSpline<TMultiSpline<SPLINETYPE>, typename SPLINETYPE::ValueType>
 {
 public:
     enum class EMappingRangeSpace: uint8
@@ -205,6 +207,12 @@ public:
     using SplineType = typename TSplineWrapper<SPLINETYPE>::SplineType;
 	typedef typename TSplineWrapper<SPLINETYPE>::ValueType ValueType;
 	
+	// Generate compile-time type ID for MultiSpline
+	DECLARE_SPLINE_TYPE_ID(
+		TEXT("MultiSpline"),
+		*TSplineValueTypeTraits<ValueType>::Name.ToString()
+	);
+	
     SplineType& GetSpline() { return TSplineWrapper<SPLINETYPE>::InternalSpline; }
     const SplineType& GetSpline() const { return TSplineWrapper<SPLINETYPE>::InternalSpline; }
 
@@ -230,7 +238,13 @@ public:
     			return false;
     		}
 
-    		// todo: compare attribute splines
+			const ISplineInterface* Child = Children[Name].Spline.Get();
+    		const ISplineInterface* OtherChild = Other.Children[Name].Spline.Get();
+
+    		if (!Child || !OtherChild || !Child->IsEqual(OtherChild))
+    		{
+    			return false;
+    		}
     	}
     	
 	    return true;
@@ -246,6 +260,17 @@ public:
 		}
     }
 
+	virtual bool IsEqual(const ISplineInterface* OtherSpline) const override
+    {
+    	if (OtherSpline->GetTypeId() == GetTypeId())
+    	{
+    		const TMultiSpline* Other = static_cast<const TMultiSpline*>(OtherSpline);
+    		return operator==(*Other);
+    	}
+    	
+    	return false;
+    }
+
 	virtual bool Serialize(FArchive& Ar) override
 	{
 		if (!TSplineWrapper<SPLINETYPE>::Serialize(Ar))
diff --git a/Engine/Source/Runtime/GeometryCore/Public/Splines/Spline.h b/Engine/Source/Runtime/GeometryCore/Public/Splines/Spline.h
--- a/Engine/Source/Runtime/GeometryCore/Public/Splines/Spline.h
+++ b/Engine/Source/Runtime/GeometryCore/Public/Splines/Spline.h
@@ -28,6 +28,7 @@ public:
 
     // ISplineInterface Forwarding
     virtual void Clear() override { InternalSpline.Clear(); }
+	virtual bool IsEqual(const ISplineInterface* OtherSpline) const override { return InternalSpline.IsEqual(OtherSpline); }
     virtual bool Serialize(FArchive& Ar) override
     {
     	// Call parent's Serialize
diff --git a/Engine/Source/Runtime/GeometryCore/Public/Splines/SplineInterfaces.h b/Engine/Source/Runtime/GeometryCore/Public/Splines/SplineInterfaces.h
--- a/Engine/Source/Runtime/GeometryCore/Public/Splines/SplineInterfaces.h
+++ b/Engine/Source/Runtime/GeometryCore/Public/Splines/SplineInterfaces.h
@@ -38,6 +38,7 @@ public:
 
     virtual void Clear() = 0;
     virtual TUniquePtr<ISplineInterface> Clone() const = 0;
+	virtual bool IsEqual(const ISplineInterface* OtherSpline) const = 0;
     
     virtual FName GetValueTypeName() const = 0;
     virtual FName GetImplementationName() const = 0;
diff --git a/Engine/Source/Runtime/GeometryCore/Public/Splines/TangentBezierSpline.h b/Engine/Source/Runtime/GeometryCore/Public/Splines/TangentBezierSpline.h
--- a/Engine/Source/Runtime/GeometryCore/Public/Splines/TangentBezierSpline.h
+++ b/Engine/Source/Runtime/GeometryCore/Public/Splines/TangentBezierSpline.h
@@ -189,6 +189,17 @@ public:
         return *this;
     }
 
+	virtual bool IsEqual(const ISplineInterface* OtherSpline) const override
+	{
+		if (OtherSpline->GetTypeId() == GetTypeId())
+		{
+			const TTangentBezierSpline* Other = static_cast<const TTangentBezierSpline*>(OtherSpline);
+			return operator==(*Other);
+		}
+		
+		return false;
+	}
+
     virtual bool Serialize(FArchive& Ar) override
     {
 		// Call immediate parent's Serialize (TSplineWrapper)
@@ -230,7 +241,10 @@ public:
 
 	bool operator==(const TTangentBezierSpline<ValueType>& Other) const
 	{
-		return InternalSpline == Other.InternalSpline;
+		return	InternalSpline == Other.InternalSpline &&
+				Tension == Other.Tension &&
+				TangentModes == Other.TangentModes &&
+				bStationaryEndpoint == Other.bStationaryEndpoint;
 	}
 
 	// Static shape generators
