Splines - Fixed bug where property change notifications were not being broadcast for both of the spline components backing spline properties.

This change is necessary to prevent splines from getting out of sync when serializing (because if one is dirty and the other is not, only one will be serialized regardless of the actual state of the data)

diff --git a/Engine/Plugins/Experimental/CineCameraRigs/Source/CineCameraRigs/Private/CineSplineComponent.cpp b/Engine/Plugins/Experimental/CineCameraRigs/Source/CineCameraRigs/Private/CineSplineComponent.cpp
--- a/Engine/Plugins/Experimental/CineCameraRigs/Source/CineCameraRigs/Private/CineSplineComponent.cpp
+++ b/Engine/Plugins/Experimental/CineCameraRigs/Source/CineCameraRigs/Private/CineSplineComponent.cpp
@@ -80,9 +80,13 @@ void UCineSplineComponent::ApplyComponentInstanceData(FCineSplineInstanceData* S
 
 			// If we are restoring the saved state, unmark the SplineCurves property as 'modified'.
 			// We don't want to consider that these changes have been made through the UCS.
-			TArray<FProperty*> Properties;
-			Properties.Emplace(FindFProperty<FProperty>(UCineSplineComponent::StaticClass(), GetSplinePropertyName()));
-			RemoveUCSModifiedProperties(Properties);
+			TArray<FProperty*> SplineProperties;
+			for (const FName& Property : GetSplinePropertyNames())
+			{
+				SplineProperties.Add(FindFProperty<FProperty>(StaticClass(), Property));
+			}
+
+			RemoveUCSModifiedProperties(SplineProperties);
 		}
 	}
 	else
diff --git a/Engine/Plugins/Experimental/CineCameraRigs/Source/CineCameraRigsEditor/Private/CineSplineMetadataDetails.cpp b/Engine/Plugins/Experimental/CineCameraRigs/Source/CineCameraRigsEditor/Private/CineSplineMetadataDetails.cpp
--- a/Engine/Plugins/Experimental/CineCameraRigs/Source/CineCameraRigsEditor/Private/CineSplineMetadataDetails.cpp
+++ b/Engine/Plugins/Experimental/CineCameraRigs/Source/CineCameraRigsEditor/Private/CineSplineMetadataDetails.cpp
@@ -97,8 +97,17 @@ void SetValues(FCineSplineMetadataDetails& Details, TArray<FInterpCurvePoint<T>>
 
 	Details.SplineComp->UpdateSpline();
 	Details.SplineComp->bSplineHasBeenEdited = true;
-	static FProperty* SplineCurvesProperty = FindFProperty<FProperty>(USplineComponent::StaticClass(), USplineComponent::GetSplinePropertyName());
-	FComponentVisualizer::NotifyPropertyModified(Details.SplineComp, SplineCurvesProperty);
+
+	static TArray<FProperty*> SplineProperties;
+	if (SplineProperties.IsEmpty())
+	{
+		for (const FName& Property : USplineComponent::GetSplinePropertyNames())
+		{
+			SplineProperties.Add(FindFProperty<FProperty>(USplineComponent::StaticClass(), Property));
+		}
+	}
+
+	FComponentVisualizer::NotifyPropertiesModified(Details.SplineComp, SplineProperties);
 	Details.Update(Details.SplineComp, Details.SelectedKeys);
 
 	if (GEditor)
diff --git a/Engine/Plugins/Experimental/Water/Source/Editor/Private/WaterSplineComponentVisualizer.cpp b/Engine/Plugins/Experimental/Water/Source/Editor/Private/WaterSplineComponentVisualizer.cpp
--- a/Engine/Plugins/Experimental/Water/Source/Editor/Private/WaterSplineComponentVisualizer.cpp
+++ b/Engine/Plugins/Experimental/Water/Source/Editor/Private/WaterSplineComponentVisualizer.cpp
@@ -571,7 +571,7 @@ bool FWaterSplineComponentVisualizer::HandleInputDelta(FEditorViewportClient* Vi
 					WaterSplineComp->bSplineHasBeenEdited = true;
 
 					// Transform the spline keys using an EPropertyChangeType::Interactive change. Later on, at the end of mouse tracking, a non-interactive change will be notified via void TrackingStopped :
-					NotifyPropertyModified(WaterSplineComp, SplineCurvesProperty, EPropertyChangeType::Interactive);
+					NotifyPropertiesModified(WaterSplineComp, SplineProperties, EPropertyChangeType::Interactive);
 
 					GEditor->RedrawLevelEditingViewports(true);
 
diff --git a/Engine/Plugins/Experimental/Water/Source/Editor/Private/WaterSplineMetadataDetails.cpp b/Engine/Plugins/Experimental/Water/Source/Editor/Private/WaterSplineMetadataDetails.cpp
--- a/Engine/Plugins/Experimental/Water/Source/Editor/Private/WaterSplineMetadataDetails.cpp
+++ b/Engine/Plugins/Experimental/Water/Source/Editor/Private/WaterSplineMetadataDetails.cpp
@@ -118,9 +118,18 @@ void SetValues(FWaterSplineMetadataDetails& Details, TArray<FInterpCurvePoint<T>
 
 	Details.SplineComp->UpdateSpline();
 	Details.SplineComp->bSplineHasBeenEdited = true;
-	static FProperty* SplineCurvesProperty = FindFProperty<FProperty>(USplineComponent::StaticClass(), USplineComponent::GetSplinePropertyName());
 	EPropertyChangeType::Type PropertyChangeType = WaterSplineMetadataDetailsLocals::IsTextCommitValueSet(CommitInfo) ? EPropertyChangeType::ValueSet : EPropertyChangeType::Interactive;
-	FComponentVisualizer::NotifyPropertyModified(Details.SplineComp, SplineCurvesProperty, PropertyChangeType);
+	
+	static TArray<FProperty*> SplineProperties;
+	if (SplineProperties.IsEmpty())
+	{
+		for (const FName& Property : USplineComponent::GetSplinePropertyNames())
+		{
+			SplineProperties.Add(FindFProperty<FProperty>(USplineComponent::StaticClass(), Property));
+		}
+	}
+
+	FComponentVisualizer::NotifyPropertiesModified(Details.SplineComp, SplineProperties, PropertyChangeType);
 	Details.Update(Details.SplineComp, Details.SelectedKeys);
 
 	GEditor->RedrawLevelEditingViewports(true);
diff --git a/Engine/Source/Editor/ComponentVisualizers/Private/SplineComponentVisualizer.cpp b/Engine/Source/Editor/ComponentVisualizers/Private/SplineComponentVisualizer.cpp
--- a/Engine/Source/Editor/ComponentVisualizers/Private/SplineComponentVisualizer.cpp
+++ b/Engine/Source/Editor/ComponentVisualizers/Private/SplineComponentVisualizer.cpp
@@ -362,8 +362,15 @@ FSplineComponentVisualizer::FSplineComponentVisualizer()
 
 	SplineComponentVisualizerActions = MakeShareable(new FUICommandList);
 
-	SplineCurvesProperty = FindFProperty<FProperty>(USplineComponent::StaticClass(), USplineComponent::GetSplinePropertyName());
-
+PRAGMA_DISABLE_DEPRECATION_WARNINGS
+	SplineCurvesProperty = nullptr;
+PRAGMA_ENABLE_DEPRECATION_WARNINGS
+	
+	for (const FName& Property : USplineComponent::GetSplinePropertyNames())
+	{
+		SplineProperties.Add(FindFProperty<FProperty>(USplineComponent::StaticClass(), Property));
+	}
+	
 	SelectionState = NewObject<USplineComponentVisualizerSelectionState>(GetTransientPackage(), TEXT("SelectionState"), RF_Transactional);
 }
 
@@ -1495,7 +1502,7 @@ bool FSplineComponentVisualizer::TransformSelectedTangent(EPropertyChangeType::T
 		SplineComp->UpdateSpline();
 		SplineComp->bSplineHasBeenEdited = true;
 
-		NotifyPropertyModified(SplineComp, SplineCurvesProperty, InPropertyChangeType);
+		NotifyPropertiesModified(SplineComp, SplineProperties, InPropertyChangeType);
 
 		return true;
 	}
@@ -1565,7 +1572,7 @@ bool FSplineComponentVisualizer::TransformSelectedKeys(EPropertyChangeType::Type
 		SplineComp->UpdateSpline();
 		SplineComp->bSplineHasBeenEdited = true;
 
-		NotifyPropertyModified(SplineComp, SplineCurvesProperty, InPropertyChangeType);
+		NotifyPropertiesModified(SplineComp, SplineProperties, InPropertyChangeType);
 
 		if (!InDeltaRotate.IsZero())
 		{
@@ -1615,7 +1622,7 @@ bool FSplineComponentVisualizer::TransformSelectedAttribute(EPropertyChangeType:
 		SelectionState->SetCachedRotation(SplineComp->GetQuaternionAtSplineInputKey(NewParam, ESplineCoordinateSpace::World));
 
 		SplineComp->bSplineHasBeenEdited = true;
-		NotifyPropertyModified(SplineComp, SplineCurvesProperty, InPropertyChangeType);
+		NotifyPropertiesModified(SplineComp, SplineProperties, InPropertyChangeType);
 		GEditor->RedrawLevelEditingViewports(true);
 
 		return true;
@@ -1953,7 +1960,7 @@ bool FSplineComponentVisualizer::HandleSnapTo(const bool bInAlign, const bool bI
 				SplineComp->UpdateSpline();
 				SplineComp->bSplineHasBeenEdited = true;
 
-				NotifyPropertyModified(SplineComp, SplineCurvesProperty);
+				NotifyPropertiesModified(SplineComp, SplineProperties);
 				if (AActor* Owner = SplineComp->GetOwner())
 				{
 					Owner->PostEditMove(true);
@@ -1981,7 +1988,7 @@ void FSplineComponentVisualizer::TrackingStopped(FEditorViewportClient* InViewpo
 	{
 		// After dragging, notify that the spline curves property has changed one last time, this time as a EPropertyChangeType::ValueSet :
 		USplineComponent* SplineComp = GetEditedSplineComponent();
-		NotifyPropertyModified(SplineComp, SplineCurvesProperty, EPropertyChangeType::ValueSet);
+		NotifyPropertiesModified(SplineComp, SplineProperties, EPropertyChangeType::ValueSet);
 		if (AActor* Owner = SplineComp->GetOwner())
 		{
 			Owner->PostEditMove(true);
@@ -2238,7 +2245,7 @@ void FSplineComponentVisualizer::SnapKeyToTransform(const ESplineComponentSnapMo
 	SplineComp->UpdateSpline();
 	SplineComp->bSplineHasBeenEdited = true;
 
-	NotifyPropertyModified(SplineComp, SplineCurvesProperty);
+	NotifyPropertiesModified(SplineComp, SplineProperties);
 	if (AActor* Owner = SplineComp->GetOwner())
 	{
 		Owner->PostEditMove(true);
@@ -2327,7 +2334,7 @@ void FSplineComponentVisualizer::OnStraightenKey(int32 Direction)
 	SplineComp->UpdateSpline();
 	SplineComp->bSplineHasBeenEdited = true;
 
-	NotifyPropertyModified(SplineComp, SplineCurvesProperty);
+	NotifyPropertiesModified(SplineComp, SplineProperties);
 	if (AActor* OwnerActor = SplineComp->GetOwner())
 	{
 		OwnerActor->PostEditMove(true);
@@ -2367,8 +2374,7 @@ void FSplineComponentVisualizer::OnToggleSnapTangentAdjustment()
 
 	SplineComp->bAdjustTangentsOnSnap = !SplineComp->bAdjustTangentsOnSnap;
 
-	TArray<FProperty*> Properties;
-	Properties.Add(SplineCurvesProperty);
+	TArray<FProperty*> Properties = SplineProperties;
 	Properties.Add(FindFProperty<FProperty>(USplineComponent::StaticClass(), GET_MEMBER_NAME_CHECKED(USplineComponent, bAdjustTangentsOnSnap)));
 	NotifyPropertiesModified(SplineComp, Properties);
 
@@ -2429,7 +2435,7 @@ void FSplineComponentVisualizer::SnapKeysToLastSelectedAxisPosition(const EAxis:
 	SplineComp->UpdateSpline();
 	SplineComp->bSplineHasBeenEdited = true;
 
-	NotifyPropertyModified(SplineComp, SplineCurvesProperty);
+	NotifyPropertiesModified(SplineComp, SplineProperties);
 	if (AActor* Owner = SplineComp->GetOwner())
 	{
 		Owner->PostEditMove(true);
@@ -2539,7 +2545,7 @@ void FSplineComponentVisualizer::OnDuplicateKey()
 	SplineComp->UpdateSpline();
 	SplineComp->bSplineHasBeenEdited = true;
 
-	NotifyPropertyModified(SplineComp, SplineCurvesProperty);
+	NotifyPropertiesModified(SplineComp, SplineProperties);
 	if (AActor* Owner = SplineComp->GetOwner())
 	{
 		Owner->PostEditMove(true);
@@ -2782,7 +2788,7 @@ void FSplineComponentVisualizer::SplitSegment(const FVector& InWorldPos, int32 I
 	SplineComp->UpdateSpline();
 	SplineComp->bSplineHasBeenEdited = true;
 
-	NotifyPropertyModified(SplineComp, SplineCurvesProperty);
+	NotifyPropertiesModified(SplineComp, SplineProperties);
 	if (AActor* Owner = SplineComp->GetOwner())
 	{
 		Owner->PostEditMove(true);
@@ -2878,7 +2884,7 @@ void FSplineComponentVisualizer::UpdateSplitSegment(const FVector& InDrag)
 	SplineComp->bSplineHasBeenEdited = true;
 
 	// Transform the spline keys using an EPropertyChangeType::Interactive change. Later on, at the end of mouse tracking, a non-interactive change will be notified via void TrackingStopped :
-	NotifyPropertyModified(SplineComp, SplineCurvesProperty, EPropertyChangeType::Interactive);
+	NotifyPropertiesModified(SplineComp, SplineProperties, EPropertyChangeType::Interactive);
 
 	GEditor->RedrawLevelEditingViewports(true);
 }
@@ -2928,7 +2934,7 @@ void FSplineComponentVisualizer::AddSegment(const FVector& InWorldPos, bool bApp
 	SplineComp->UpdateSpline();
 	SplineComp->bSplineHasBeenEdited = true;
 
-	NotifyPropertyModified(SplineComp, SplineCurvesProperty);
+	NotifyPropertiesModified(SplineComp, SplineProperties);
 
 	GEditor->RedrawLevelEditingViewports(true);
 }
@@ -2961,7 +2967,7 @@ void FSplineComponentVisualizer::UpdateAddSegment(const FVector& InDrag)
 	SplineComp->bSplineHasBeenEdited = true;
 
 	// Transform the spline keys using an EPropertyChangeType::Interactive change. Later on, at the end of mouse tracking, a non-interactive change will be notified via void TrackingStopped :
-	NotifyPropertyModified(SplineComp, SplineCurvesProperty, EPropertyChangeType::Interactive);
+	NotifyPropertiesModified(SplineComp, SplineProperties, EPropertyChangeType::Interactive);
 
 	GEditor->RedrawLevelEditingViewports(true);
 }
@@ -3019,7 +3025,7 @@ void FSplineComponentVisualizer::OnDeleteKey()
 	SplineComp->UpdateSpline();
 	SplineComp->bSplineHasBeenEdited = true;
 
-	NotifyPropertyModified(SplineComp, SplineCurvesProperty);
+	NotifyPropertiesModified(SplineComp, SplineProperties);
 	if (AActor* OwnerActor = SplineComp->GetOwner())
 	{
 		OwnerActor->PostEditMove(true);
@@ -3096,7 +3102,7 @@ void FSplineComponentVisualizer::OnResetToAutomaticTangent(EInterpCurveMode Mode
 		SplineComp->UpdateSpline();
 		SplineComp->bSplineHasBeenEdited = true;
 
-		NotifyPropertyModified(SplineComp, SplineCurvesProperty);
+		NotifyPropertiesModified(SplineComp, SplineProperties);
 		if (AActor* OwnerActor = SplineComp->GetOwner())
 		{
 			OwnerActor->PostEditMove(true);
@@ -3156,7 +3162,7 @@ void FSplineComponentVisualizer::OnSetKeyType(EInterpCurveMode Mode)
 		SplineComp->UpdateSpline();
 		SplineComp->bSplineHasBeenEdited = true;
 
-		NotifyPropertyModified(SplineComp, SplineCurvesProperty);
+		NotifyPropertiesModified(SplineComp, SplineProperties);
 		if (AActor* OwnerActor = SplineComp->GetOwner())
 		{
 			OwnerActor->PostEditMove(true);
@@ -3246,8 +3252,7 @@ void FSplineComponentVisualizer::OnSetDiscontinuousSpline()
 		}
 	}
 
-	TArray<FProperty*> Properties;
-	Properties.Add(SplineCurvesProperty);
+	TArray<FProperty*> Properties = SplineProperties;
 	Properties.Add(FindFProperty<FProperty>(USplineComponent::StaticClass(), GET_MEMBER_NAME_CHECKED(USplineComponent, bAllowDiscontinuousSpline)));
 	NotifyPropertiesModified(SplineComp, Properties);
 
@@ -3278,8 +3283,7 @@ void FSplineComponentVisualizer::OnToggleClosedLoop()
 
 	SplineComp->SetClosedLoop(!SplineComp->IsClosedLoop());
 
-	TArray<FProperty*> Properties;
-	Properties.Add(SplineCurvesProperty);
+	TArray<FProperty*> Properties = SplineProperties;
 	Properties.Add(FindFProperty<FProperty>(USplineComponent::StaticClass(), GET_MEMBER_NAME_CHECKED(USplineComponent, bClosedLoop)));
 	NotifyPropertiesModified(SplineComp, Properties);
 
@@ -3369,7 +3373,7 @@ void FSplineComponentVisualizer::OnAddAttributeKey()
 	SelectionState->SetCachedRotation(SplineComp->GetQuaternionAtSplineInputKey(NearestParam, ESplineCoordinateSpace::World));
 
 	SplineComp->bSplineHasBeenEdited = true;
-	NotifyPropertyModified(SplineComp, SplineCurvesProperty, EPropertyChangeType::ArrayAdd);
+	NotifyPropertiesModified(SplineComp, SplineProperties, EPropertyChangeType::ArrayAdd);
 }
 
 bool FSplineComponentVisualizer::CanAddAttributeKey() const
@@ -3415,7 +3419,7 @@ void FSplineComponentVisualizer::OnDeleteAttributeKey()
 	}
 
 	SplineComp->bSplineHasBeenEdited = true;
-	NotifyPropertyModified(SplineComp, SplineCurvesProperty, EPropertyChangeType::ArrayRemove);
+	NotifyPropertiesModified(SplineComp, SplineProperties, EPropertyChangeType::ArrayRemove);
 	
 	// Necessary because the attribute proxies store invalid indices, we must redraw them.
 	GEditor->RedrawAllViewports();
@@ -3872,7 +3876,7 @@ TSharedPtr<SWidget> FSplineComponentVisualizer::GenerateAttributeEditorWidget()
 					}
 
 					LocalSplineComp->bSplineHasBeenEdited = true;
-					NotifyPropertyModified(LocalSplineComp, SplineCurvesProperty, EPropertyChangeType::ValueSet);
+					NotifyPropertiesModified(LocalSplineComp, SplineProperties, EPropertyChangeType::ValueSet);
 					GEditor->RedrawAllViewports();
 				}
 			}
@@ -3915,7 +3919,7 @@ TSharedPtr<SWidget> FSplineComponentVisualizer::GenerateAttributeEditorWidget()
 				LocalSplineComp->SetFloatPropertyAtIndex(AttributeIndex, NewValue, AttributeName);
 				
 				LocalSplineComp->bSplineHasBeenEdited = true;
-				NotifyPropertyModified(LocalSplineComp, SplineCurvesProperty, EPropertyChangeType::ValueSet);
+				NotifyPropertiesModified(LocalSplineComp, SplineProperties, EPropertyChangeType::ValueSet);
 				GEditor->RedrawAllViewports();
 			}
 		}
diff --git a/Engine/Source/Editor/ComponentVisualizers/Private/SplineGeneratorPanel.cpp b/Engine/Source/Editor/ComponentVisualizers/Private/SplineGeneratorPanel.cpp
--- a/Engine/Source/Editor/ComponentVisualizers/Private/SplineGeneratorPanel.cpp
+++ b/Engine/Source/Editor/ComponentVisualizers/Private/SplineGeneratorPanel.cpp
@@ -90,9 +90,14 @@ void USplineGeneratorBase::Reset()
 		SelectedSplineComponent->UpdateSpline();
 		SelectedSplineComponent->bSplineHasBeenEdited = true;
 
-		FProperty* SplineCurvesProperty = FindFProperty<FProperty>(USplineComponent::StaticClass(), USplineComponent::GetSplinePropertyName());
-		FComponentVisualizer::NotifyPropertyModified(SelectedSplineComponent.Get(), SplineCurvesProperty);
-
+		TArray<FProperty*> SplineProperties;
+		for (const FName& Property : USplineComponent::GetSplinePropertyNames())
+		{
+			SplineProperties.Add(FindFProperty<FProperty>(USplineComponent::StaticClass(), Property));
+		}
+		
+		FComponentVisualizer::NotifyPropertiesModified(SelectedSplineComponent.Get(), SplineProperties);
+		
 		GEditor->RedrawLevelEditingViewports(true);
 	}
 }
@@ -120,8 +125,13 @@ void USplineGeneratorBase::PreviewCurve()
 		SelectedSplineComponent->UpdateSpline();
 		SelectedSplineComponent->bSplineHasBeenEdited = true;
 
-		FProperty* SplineCurvesProperty = FindFProperty<FProperty>(USplineComponent::StaticClass(), USplineComponent::GetSplinePropertyName());
-		FComponentVisualizer::NotifyPropertyModified(SelectedSplineComponent.Get(), SplineCurvesProperty);
+		TArray<FProperty*> SplineProperties;
+		for (const FName& Property : USplineComponent::GetSplinePropertyNames())
+		{
+			SplineProperties.Add(FindFProperty<FProperty>(USplineComponent::StaticClass(), Property));
+		}
+		
+		FComponentVisualizer::NotifyPropertiesModified(SelectedSplineComponent.Get(), SplineProperties);
 
 		GEditor->RedrawLevelEditingViewports(true);
 	}
diff --git a/Engine/Source/Editor/ComponentVisualizers/Public/SplineComponentVisualizer.h b/Engine/Source/Editor/ComponentVisualizers/Public/SplineComponentVisualizer.h
--- a/Engine/Source/Editor/ComponentVisualizers/Public/SplineComponentVisualizer.h
+++ b/Engine/Source/Editor/ComponentVisualizers/Public/SplineComponentVisualizer.h
@@ -525,8 +525,12 @@ protected:
 	/** Snap: Snap to actor mode. */
 	ESplineComponentSnapMode SnapToActorMode;
 
+	UE_DEPRECATED(5.7, "SplineCurvesProperty is deprecated, please use SplineProperties.")
 	FProperty* SplineCurvesProperty;
 
+	/** The set of properties to mark as dirty when manipulating the target spline component. */
+	TArray<FProperty*> SplineProperties;
+	
 	FDelegateHandle DeselectedInEditorDelegateHandle;
 
 	void UpdateSharedAttributeNames() const
diff --git a/Engine/Source/Editor/DetailCustomizations/Private/SplineComponentDetails.cpp b/Engine/Source/Editor/DetailCustomizations/Private/SplineComponentDetails.cpp
--- a/Engine/Source/Editor/DetailCustomizations/Private/SplineComponentDetails.cpp
+++ b/Engine/Source/Editor/DetailCustomizations/Private/SplineComponentDetails.cpp
@@ -321,7 +321,7 @@ private:
 	TSharedValue<ESplinePointType::Type> PointType;
 
 	TSharedPtr<FSplineComponentVisualizer> SplineVisualizer;
-	FProperty* SplineCurvesProperty;
+	TArray<FProperty*> SplineProperties;
 	TArray<TSharedPtr<FString>> SplinePointTypes;
 	TSharedPtr<ISplineMetadataDetails> SplineMetaDataDetails;
 	FSimpleDelegate OnRegenerateChildren;
@@ -341,8 +341,11 @@ FSplinePointDetails::FSplinePointDetails(USplineComponent* InOwningSplineCompone
 	SplineVisualizer = StaticCastSharedPtr<FSplineComponentVisualizer>(Visualizer);
 	check(SplineVisualizer.IsValid());
 
-	SplineCurvesProperty = FindFProperty<FProperty>(USplineComponent::StaticClass(), USplineComponent::GetSplinePropertyName());
-
+	for (const FName& Property : USplineComponent::GetSplinePropertyNames())
+	{
+		SplineProperties.Add(FindFProperty<FProperty>(USplineComponent::StaticClass(), Property));
+	}
+	
 	const TArray<ESplinePointType::Type> EnabledSplinePointTypes = InOwningSplineComponent->GetEnabledSplinePointTypes();
 
 	UEnum* SplinePointTypeEnum = StaticEnum<ESplinePointType::Type>();
@@ -978,7 +981,7 @@ PRAGMA_DISABLE_DEPRECATION_WARNINGS
 
 	SplineComp->UpdateSpline();
 	SplineComp->bSplineHasBeenEdited = true;
-	FComponentVisualizer::NotifyPropertyModified(SplineComp, SplineCurvesProperty);
+	FComponentVisualizer::NotifyPropertiesModified(SplineComp, SplineProperties);
 	if (AActor* Owner = SplineComp->GetOwner())
 	{
 		Owner->PostEditMove(true);
@@ -1024,7 +1027,7 @@ void FSplinePointDetails::OnSetPosition(FVector::FReal NewValue, ETextCommit::Ty
 	{
 		SplineComp->UpdateSpline();
 		SplineComp->bSplineHasBeenEdited = true;
-		FComponentVisualizer::NotifyPropertyModified(SplineComp, SplineCurvesProperty, EPropertyChangeType::ValueSet);
+		FComponentVisualizer::NotifyPropertiesModified(SplineComp, SplineProperties, EPropertyChangeType::ValueSet);
 		if (AActor* Owner = SplineComp->GetOwner())
 		{
 			Owner->PostEditMove(true);
@@ -1071,7 +1074,7 @@ void FSplinePointDetails::OnSetArriveTangent(FVector::FReal NewValue, ETextCommi
 	{
 		SplineComp->UpdateSpline();
 		SplineComp->bSplineHasBeenEdited = true;
-		FComponentVisualizer::NotifyPropertyModified(SplineComp, SplineCurvesProperty, EPropertyChangeType::ValueSet);
+		FComponentVisualizer::NotifyPropertiesModified(SplineComp, SplineProperties, EPropertyChangeType::ValueSet);
 		if (AActor* Owner = SplineComp->GetOwner())
 		{
 			Owner->PostEditMove(true);
@@ -1117,7 +1120,7 @@ void FSplinePointDetails::OnSetLeaveTangent(FVector::FReal NewValue, ETextCommit
 	{
 		SplineComp->UpdateSpline();
 		SplineComp->bSplineHasBeenEdited = true;
-		FComponentVisualizer::NotifyPropertyModified(SplineComp, SplineCurvesProperty, EPropertyChangeType::ValueSet);
+		FComponentVisualizer::NotifyPropertiesModified(SplineComp, SplineProperties, EPropertyChangeType::ValueSet);
 		if (AActor* Owner = SplineComp->GetOwner())
 		{
 			Owner->PostEditMove(true);
@@ -1192,7 +1195,7 @@ void FSplinePointDetails::OnSetRotation(FRotator::FReal NewValue, ETextCommit::T
 	{
 		SplineComp->UpdateSpline();
 		SplineComp->bSplineHasBeenEdited = true;
-		FComponentVisualizer::NotifyPropertyModified(SplineComp, SplineCurvesProperty, EPropertyChangeType::ValueSet);
+		FComponentVisualizer::NotifyPropertiesModified(SplineComp, SplineProperties, EPropertyChangeType::ValueSet);
 		if (AActor* Owner = SplineComp->GetOwner())
 		{
 			Owner->PostEditMove(true);
@@ -1235,7 +1238,7 @@ void FSplinePointDetails::OnSetScale(FVector::FReal NewValue, ETextCommit::Type
 	{
 		SplineComp->UpdateSpline();
 		SplineComp->bSplineHasBeenEdited = true;
-		FComponentVisualizer::NotifyPropertyModified(SplineComp, SplineCurvesProperty, EPropertyChangeType::ValueSet);
+		FComponentVisualizer::NotifyPropertiesModified(SplineComp, SplineProperties, EPropertyChangeType::ValueSet);
 		if (AActor* Owner = SplineComp->GetOwner())
 		{
 			Owner->PostEditMove(true);
@@ -1305,7 +1308,7 @@ void FSplinePointDetails::OnSplinePointTypeChanged(TSharedPtr<FString> NewValue,
 	{
 		SplineComp->UpdateSpline();
 		SplineComp->bSplineHasBeenEdited = true;
-		FComponentVisualizer::NotifyPropertyModified(SplineComp, SplineCurvesProperty);
+		FComponentVisualizer::NotifyPropertiesModified(SplineComp, SplineProperties);
 		if (AActor* Owner = SplineComp->GetOwner())
 		{
 			Owner->PostEditMove(true);
@@ -1808,10 +1811,12 @@ TSharedRef<IDetailCustomization> FSplineComponentDetails::MakeInstance()
 
 void FSplineComponentDetails::CustomizeDetails(IDetailLayoutBuilder& DetailBuilder)
 {
-	// Hide the SplineCurves property
-	TSharedPtr<IPropertyHandle> SplineCurvesProperty = DetailBuilder.GetProperty(USplineComponent::GetSplinePropertyName());
-	SplineCurvesProperty->MarkHiddenByCustomization();
-
+	// Hide all spline properties
+	for (FName& PropertyName : USplineComponent::GetSplinePropertyNames())
+	{
+		TSharedPtr<IPropertyHandle> Property = DetailBuilder.GetProperty(PropertyName);
+		Property->MarkHiddenByCustomization();
+	}
 
 	TArray<TWeakObjectPtr<UObject>> ObjectsBeingCustomized;
 	DetailBuilder.GetObjectsBeingCustomized(ObjectsBeingCustomized);
diff --git a/Engine/Source/Runtime/Engine/Classes/Components/SplineComponent.h b/Engine/Source/Runtime/Engine/Classes/Components/SplineComponent.h
--- a/Engine/Source/Runtime/Engine/Classes/Components/SplineComponent.h
+++ b/Engine/Source/Runtime/Engine/Classes/Components/SplineComponent.h
@@ -219,9 +219,11 @@ public:
 	UFUNCTION(CallInEditor, Category="Spline")
 	void PopulateFromLegacy();
 
-	// Gives the name of the backing spline property to allow users
-	// the ability to correctly trigger property change events.
+	/** Return the name of the authoritative spline property. */
 	ENGINE_API static FName GetSplinePropertyName();
+
+	/** Return the set of names of the backing spline properties. */
+	ENGINE_API static TSet<FName> GetSplinePropertyNames();
 	
 	UPROPERTY(EditAnywhere, Replicated, Category=Points)
 	FSplineCurves SplineCurves;
diff --git a/Engine/Source/Runtime/Engine/Private/Components/SplineComponent.cpp b/Engine/Source/Runtime/Engine/Private/Components/SplineComponent.cpp
--- a/Engine/Source/Runtime/Engine/Private/Components/SplineComponent.cpp
+++ b/Engine/Source/Runtime/Engine/Private/Components/SplineComponent.cpp
@@ -72,6 +72,13 @@ PRAGMA_ENABLE_DEPRECATION_WARNINGS
 	}
 }
 
+TSet<FName> USplineComponent::GetSplinePropertyNames()
+{
+PRAGMA_DISABLE_DEPRECATION_WARNINGS
+	return {GET_MEMBER_NAME_CHECKED(USplineComponent, SplineCurves), GET_MEMBER_NAME_CHECKED(USplineComponent, Spline)};
+PRAGMA_ENABLE_DEPRECATION_WARNINGS
+}
+
 USplineComponent::USplineComponent(const FObjectInitializer& ObjectInitializer)
 	: Super(ObjectInitializer)
 #if WITH_EDITORONLY_DATA
@@ -2665,9 +2672,13 @@ void USplineComponent::ApplyComponentInstanceData(FSplineInstanceData* SplineIns
 
 			// If we are restoring the saved state, unmark the SplineCurves property as 'modified'.
 			// We don't want to consider that these changes have been made through the UCS.
-			TArray<FProperty*> Properties;
-			Properties.Emplace(FindFProperty<FProperty>(StaticClass(), GetSplinePropertyName()));
-			RemoveUCSModifiedProperties(Properties);
+			TArray<FProperty*> SplineProperties;
+			for (const FName& Property : GetSplinePropertyNames())
+			{
+				SplineProperties.Add(FindFProperty<FProperty>(StaticClass(), Property));
+			}
+			
+			RemoveUCSModifiedProperties(SplineProperties);
 		}
 	}
 	else
@@ -2705,9 +2716,13 @@ void USplineComponent::ApplyComponentInstanceData(struct FSplineComponentInstanc
 
 			// If we are restoring the saved state, unmark the SplineCurves property as 'modified'.
 			// We don't want to consider that these changes have been made through the UCS.
-			TArray<FProperty*> Properties;
-			Properties.Emplace(FindFProperty<FProperty>(StaticClass(), GetSplinePropertyName()));
-			RemoveUCSModifiedProperties(Properties);
+			TArray<FProperty*> SplineProperties;
+			for (const FName& Property : GetSplinePropertyNames())
+			{
+				SplineProperties.Add(FindFProperty<FProperty>(StaticClass(), Property));
+			}
+			
+			RemoveUCSModifiedProperties(SplineProperties);
 		}
 	}
 	else
