ComponentVisualizersModule: Add functionality for muting/unmuting a component visualizer by name.

We can currently unregister a visualizer, but we cannot reregister it without re-instantiating it or keeping the originally registered instance alive. This allows users to very simply temporarily enable/disable a visualizer without having to do any real work.

diff --git a/Engine/Source/Editor/ComponentVisualizers/Private/ComponentVisualizers.cpp b/Engine/Source/Editor/ComponentVisualizers/Private/ComponentVisualizers.cpp
--- a/Engine/Source/Editor/ComponentVisualizers/Private/ComponentVisualizers.cpp
+++ b/Engine/Source/Editor/ComponentVisualizers/Private/ComponentVisualizers.cpp
@@ -92,3 +92,41 @@ void FComponentVisualizersModule::RegisterComponentVisualizer(FName ComponentCla
 		Visualizer->OnRegister();
 	}
 }
+
+void FComponentVisualizersModule::MuteComponentVisualizer(FName ComponentClassName)
+{
+	if (!GUnrealEd || !RegisteredComponentClassNames.Contains(ComponentClassName) || MutedComponentVisualizers.Contains(ComponentClassName))
+	{
+		return;
+	}
+	
+	if (const TSharedPtr<FComponentVisualizer> Visualizer = GUnrealEd->FindComponentVisualizer(ComponentClassName))
+	{
+		// Unregister the visualizer.
+		GUnrealEd->UnregisterComponentVisualizer(ComponentClassName);
+		RegisteredComponentClassNames.Remove(ComponentClassName);
+
+		// Cache the visualizer.
+		MutedComponentVisualizers.Add(ComponentClassName, Visualizer);
+	}
+}
+
+void FComponentVisualizersModule::UnmuteComponentVisualizer(FName ComponentClassName)
+{
+	if (!GUnrealEd || RegisteredComponentClassNames.Contains(ComponentClassName) || !MutedComponentVisualizers.Contains(ComponentClassName))
+	{
+		return;
+	}
+
+	if (const TSharedPtr<FComponentVisualizer>* VisualizerPtr = MutedComponentVisualizers.Find(ComponentClassName))
+	{
+		const TSharedPtr<FComponentVisualizer> Visualizer = *VisualizerPtr;
+
+		// Register the visualizer, but don't call OnRegister (already called previously).
+		GUnrealEd->RegisterComponentVisualizer(ComponentClassName, Visualizer);
+		RegisteredComponentClassNames.Add(ComponentClassName);
+
+		// Uncache the visualizer.
+		MutedComponentVisualizers.Remove(ComponentClassName);
+	}
+}
\ No newline at end of file
diff --git a/Engine/Source/Editor/ComponentVisualizers/Public/ComponentVisualizers.h b/Engine/Source/Editor/ComponentVisualizers/Public/ComponentVisualizers.h
--- a/Engine/Source/Editor/ComponentVisualizers/Public/ComponentVisualizers.h
+++ b/Engine/Source/Editor/ComponentVisualizers/Public/ComponentVisualizers.h
@@ -33,13 +33,20 @@ public:
 		return true;
 	}
 
-	/** Register a visualizer for a particular componen class */
+	/** Register a visualizer for a particular component class */
 	COMPONENTVISUALIZERS_API void RegisterComponentVisualizer(FName ComponentClassName, TSharedPtr<FComponentVisualizer> Visualizer);
 
-private:
+	/** Mutes a registered visualizer so that it will no longer be used. Can be restored without re-registering. */
+	COMPONENTVISUALIZERS_API void MuteComponentVisualizer(FName ComponentClassName);
+
+	/** Unmutes a muted visualizer so that it will be used again by the editor. */
+	COMPONENTVISUALIZERS_API void UnmuteComponentVisualizer(FName ComponentClassName);
 
-	
+private:
 
 	/** Array of component class names we have registered, so we know what to unregister afterwards */
 	TArray<FName> RegisteredComponentClassNames;
+
+	/** Set of muted component visualizers. We must retain the shared pointer so we can unregister from the editor. */
+	TMap<FName, TSharedPtr<FComponentVisualizer>> MutedComponentVisualizers;
 };
