SplineComponentVisualizer: Add suppressor feature interface

diff --git a/Engine/Source/Editor/ComponentVisualizers/Private/SplineComponentVisualizer.cpp b/Engine/Source/Editor/ComponentVisualizers/Private/SplineComponentVisualizer.cpp
--- a/Engine/Source/Editor/ComponentVisualizers/Private/SplineComponentVisualizer.cpp
+++ b/Engine/Source/Editor/ComponentVisualizers/Private/SplineComponentVisualizer.cpp
@@ -36,6 +36,7 @@
 #include "Widgets/Input/SEditableTextBox.h"
 #include "Widgets/Input/SNumericEntryBox.h"
 #include "Widgets/Views/SListView.h"
+#include "Features/IModularFeatures.h"
 
 #include UE_INLINE_GENERATED_CPP_BY_NAME(SplineComponentVisualizer)
 
@@ -52,6 +53,18 @@ DEFINE_LOG_CATEGORY_STATIC(LogSplineComponentVisualizer, Log, All)
 
 namespace SplineComponentVisualizerLocals
 {
+	bool IsEnabledForSpline(const USplineComponent* InSplineComponent)
+	{
+		auto ModularFeatureName = ISplineComponentVisualizerSuppressor::GetModularFeatureName();
+	
+		if (IModularFeatures::Get().IsModularFeatureAvailable(ModularFeatureName))
+		{
+			return !IModularFeatures::Get().GetModularFeature<ISplineComponentVisualizerSuppressor>(ModularFeatureName).ShouldSuppress(InSplineComponent);
+		}
+		
+		return true;
+	}
+	
 	// This is mostly modeled on FindNearestVisibleObjectHit_Internal in ModelingSceneSnappingManager.cpp,
 	// which we probably can't access because it lives in a plugin.
 	// TODO: Perhaps this code should live in some common utility place? Certainly if we do this again.
@@ -648,6 +661,11 @@ void FSplineComponentVisualizer::DrawVisualization(const UActorComponent* Compon
 {
 	if (const USplineComponent* SplineComp = Cast<const USplineComponent>(Component))
 	{
+		if (!SplineComponentVisualizerLocals::IsEnabledForSpline(SplineComp))
+		{
+			return;
+		}
+		
 		const FInterpCurveVector& SplineInfo = SplineComp->GetSplinePointsPosition();
 		const USplineComponent* EditedSplineComp = GetEditedSplineComponent();
 
@@ -892,6 +910,11 @@ void FSplineComponentVisualizer::DrawVisualizationHUD(const UActorComponent* Com
 {
 	if (const USplineComponent* SplineComp = Cast<const USplineComponent>(Component))
 	{
+		if (!SplineComponentVisualizerLocals::IsEnabledForSpline(SplineComp))
+		{
+			return;
+		}
+		
 		const bool bIsSplineEditable = !SplineComp->bModifiedByConstructionScript; // bSplineHasBeenEdited || SplineInfo == Archetype->SplineCurves.Position || SplineComp->bInputSplinePointsToConstructionScript;
 		const USplineComponent* EditedSplineComp = GetEditedSplineComponent();
 
@@ -1042,6 +1065,13 @@ bool FSplineComponentVisualizer::VisProxyHandleClick(FEditorViewportClient* InVi
 	{
 		check(SelectionState);
 
+		if (const USplineComponent* SplineComp = CastChecked<const USplineComponent>(VisProxy->Component.Get());
+			!SplineComp || (SplineComp && !SplineComponentVisualizerLocals::IsEnabledForSpline(SplineComp)))
+		{
+			// Consume no input if this visualizer is currently being suppressed.
+			return false;
+		}
+		
 		if (VisProxy->IsA(HSplineKeyProxy::StaticGetType()))
 		{
 			// Control point clicked
@@ -1231,6 +1261,11 @@ bool FSplineComponentVisualizer::GetWidgetLocation(const FEditorViewportClient*
 	USplineComponent* SplineComp = GetEditedSplineComponent();
 	if (SplineComp != nullptr)
 	{
+		if (!SplineComponentVisualizerLocals::IsEnabledForSpline(SplineComp))
+		{
+			return false;
+		}
+		
 		check(SelectionState);
 
 		const int32 NumPoints = SplineComp->GetNumberOfSplinePoints();
@@ -1301,6 +1336,11 @@ bool FSplineComponentVisualizer::GetCustomInputCoordinateSystem(const FEditorVie
 		USplineComponent* SplineComp = GetEditedSplineComponent();
 		if (SplineComp != nullptr)
 		{
+			if (!SplineComponentVisualizerLocals::IsEnabledForSpline(SplineComp))
+			{
+				return false;
+			}
+			
 			check(SelectionState);
 			OutMatrix = FRotationMatrix::Make(SelectionState->GetCachedRotation());
 			return true;
@@ -1362,6 +1402,11 @@ bool FSplineComponentVisualizer::HandleInputDelta(FEditorViewportClient* Viewpor
 	bool bInputHandled = false;
 	if (SplineComp != nullptr)
 	{
+		if (!SplineComponentVisualizerLocals::IsEnabledForSpline(SplineComp))
+		{
+			return false;
+		}
+		
 		if (IsAnySelectedKeyIndexOutOfRange(SplineComp))
 		{
 			// Something external has changed the number of spline points, meaning that the cached selected keys are no longer valid
@@ -1638,11 +1683,19 @@ bool FSplineComponentVisualizer::HandleInputKey(FEditorViewportClient* ViewportC
 	bool bHandled = false;
 
 	USplineComponent* SplineComp = GetEditedSplineComponent();
-	if (SplineComp != nullptr && IsAnySelectedKeyIndexOutOfRange(SplineComp))
+	if (SplineComp)
 	{
-		// Something external has changed the number of spline points, meaning that the cached selected keys are no longer valid
-		EndEditing();
-		return false;
+		if (IsAnySelectedKeyIndexOutOfRange(SplineComp))
+		{
+			// Something external has changed the number of spline points, meaning that the cached selected keys are no longer valid
+			EndEditing();
+			return false;
+		}
+
+		if (!SplineComponentVisualizerLocals::IsEnabledForSpline(SplineComp))
+		{
+			return false;
+		}
 	}
 
 	if (Key == EKeys::LeftMouseButton && Event == IE_Released)
@@ -1688,6 +1741,12 @@ bool FSplineComponentVisualizer::HandleInputKey(FEditorViewportClient* ViewportC
 
 bool FSplineComponentVisualizer::HandleModifiedClick(FEditorViewportClient* InViewportClient, HHitProxy* HitProxy, const FViewportClick& Click)
 {
+	if (const USplineComponent* SplineComp = GetEditedSplineComponent();
+		SplineComp && !SplineComponentVisualizerLocals::IsEnabledForSpline(SplineComp))
+	{
+		return false;
+	}
+	
 	if (Click.IsControlDown())
 	{
 		ESplineComponentSnapMode SnapMode;
@@ -1779,6 +1838,11 @@ bool FSplineComponentVisualizer::HandleBoxSelect(const FBox& InBox, FEditorViewp
 	USplineComponent* SplineComp = GetEditedSplineComponent();
 	if (SplineComp != nullptr)
 	{
+		if (!SplineComponentVisualizerLocals::IsEnabledForSpline(SplineComp))
+		{
+			return false;
+		}
+		
 		bool bSelectionChanged = false;
 		bool bAppendToSelection = InViewportClient->IsShiftPressed();
 
@@ -1823,6 +1887,11 @@ bool FSplineComponentVisualizer::HandleFrustumSelect(const FConvexVolume& InFrus
 	USplineComponent* SplineComp = GetEditedSplineComponent();
 	if (SplineComp != nullptr)
 	{
+		if (!SplineComponentVisualizerLocals::IsEnabledForSpline(SplineComp))
+		{
+			return false;
+		}
+		
 		bool bSelectionChanged = false;
 		bool bAppendToSelection = InViewportClient->IsShiftPressed();
 
@@ -1863,6 +1932,11 @@ bool FSplineComponentVisualizer::HasFocusOnSelectionBoundingBox(FBox& OutBoundin
 
 	if (USplineComponent* SplineComp = GetEditedSplineComponent())
 	{
+		if (!SplineComponentVisualizerLocals::IsEnabledForSpline(SplineComp))
+		{
+			return false;
+		}
+		
 		check(SelectionState);
 		const TSet<int32>& SelectedKeys = SelectionState->GetSelectedKeys();
 		if (SelectedKeys.Num() > 0)
@@ -3613,6 +3687,12 @@ bool FSplineComponentVisualizer::CanSelectSplinePoints() const
 
 TSharedPtr<SWidget> FSplineComponentVisualizer::GenerateContextMenu() const
 {
+	if (const USplineComponent* SplineComp = GetEditedSplineComponent();
+		SplineComp && !SplineComponentVisualizerLocals::IsEnabledForSpline(SplineComp))
+	{
+		return nullptr;
+	}
+	
 	FMenuBuilder MenuBuilder(true, SplineComponentVisualizerActions);
 	
 	GenerateContextMenuSections(MenuBuilder);
diff --git a/Engine/Source/Editor/ComponentVisualizers/Public/SplineComponentVisualizer.h b/Engine/Source/Editor/ComponentVisualizers/Public/SplineComponentVisualizer.h
--- a/Engine/Source/Editor/ComponentVisualizers/Public/SplineComponentVisualizer.h
+++ b/Engine/Source/Editor/ComponentVisualizers/Public/SplineComponentVisualizer.h
@@ -26,6 +26,7 @@
 #include "UObject/Object.h"
 #include "UObject/ObjectMacros.h"
 #include "UObject/UObjectGlobals.h"
+#include "Features/IModularFeature.h"
 
 #include "SplineComponentVisualizer.generated.h"
 
@@ -552,4 +553,21 @@ private:
 	static UE_API TWeakPtr<SWindow> WeakExistingWindow;
 };
 
+/** This class could be changed or removed without deprecation, use at your own risk. */
+class UE_INTERNAL ISplineComponentVisualizerSuppressor;
+
+PRAGMA_DISABLE_INTERNAL_WARNINGS
+/** Allows implements to suppress the visualizer (disabling rendering, input capture, etc.) for a particular spline component. */
+class ISplineComponentVisualizerSuppressor : public IModularFeature
+{
+public:
+
+	virtual ~ISplineComponentVisualizerSuppressor() = default;
+	
+	static FName GetModularFeatureName() { return "SplineComponentVisualizerSuppressor"; }
+	
+	virtual bool ShouldSuppress(const USplineComponent* SplineComponent) const = 0;
+};
+PRAGMA_ENABLE_INTERNAL_WARNINGS
+
 #undef UE_API
