SplineComponent: Fix metadata insertion

Metadata insertion behaves okay for alt-drag segment splits, but right click segment splits were not properly computing metadata values.

Now, spline components do a lerp using alpha = 0.5 by default when inserting a point (previously it was just wrong) in addition to the spline component visualizer getting a more accurate approximation using the segment parameter as the lerp factor

diff --git a/Engine/Source/Editor/ComponentVisualizers/Private/SplineComponentVisualizer.cpp b/Engine/Source/Editor/ComponentVisualizers/Private/SplineComponentVisualizer.cpp
--- a/Engine/Source/Editor/ComponentVisualizers/Private/SplineComponentVisualizer.cpp
+++ b/Engine/Source/Editor/ComponentVisualizers/Private/SplineComponentVisualizer.cpp
@@ -2858,6 +2858,12 @@ void FSplineComponentVisualizer::SplitSegment(const FVector& InWorldPos, int32 I
 	SplineComp->SetScaleAtSplinePoint(SegmentSplitIndex, NewScale, false);
 	SplineComp->SetSplinePointType(SegmentSplitIndex, NewSplinePointType, false);
 
+	// Update metadata
+	if (USplineMetadata* SplineMetadata = SplineComp->GetSplinePointsMetadata())
+	{
+		SplineMetadata->UpdatePoint(SegmentSplitIndex, t, SplineComp->IsClosedLoop());
+	}
+	
 	// Set selection to new key
 	ChangeSelectionState(SegmentSplitIndex, false);
 
diff --git a/Engine/Source/Runtime/Engine/Private/Components/SplineComponent.cpp b/Engine/Source/Runtime/Engine/Private/Components/SplineComponent.cpp
--- a/Engine/Source/Runtime/Engine/Private/Components/SplineComponent.cpp
+++ b/Engine/Source/Runtime/Engine/Private/Components/SplineComponent.cpp
@@ -1070,7 +1070,7 @@ void USplineComponent::AddSplinePointAtIndex(const FVector& Position, int32 Inde
 		
 		if (USplineMetadata* Metadata = GetSplinePointsMetadata())
 		{
-			Metadata->InsertPoint(InKey, Index, bClosedLoop);
+			Metadata->InsertPoint(Index, 0.5f, bClosedLoop);
 		}
 		
 		// Adjust subsequent points' input keys to make room for the value just added
