SplineComponentVisualizer: Use existing attribute value when adding new points

diff --git a/Engine/Source/Editor/ComponentVisualizers/Private/SplineComponentVisualizer.cpp b/Engine/Source/Editor/ComponentVisualizers/Private/SplineComponentVisualizer.cpp
--- a/Engine/Source/Editor/ComponentVisualizers/Private/SplineComponentVisualizer.cpp
+++ b/Engine/Source/Editor/ComponentVisualizers/Private/SplineComponentVisualizer.cpp
@@ -3466,8 +3466,8 @@ void FSplineComponentVisualizer::OnAddAttributeKey()
 	}
 	
 	const float NearestParam = SplineComp->FindInputKeyClosestToWorldLocation(SelectionState->GetSelectedSplinePosition());
-	// const float 
-	int32 NewAttributeIndex = SplineComp->SetFloatPropertyAtSplineInputKey(NearestParam, 0.f, SelectionState->GetSelectedAttributeName());
+	const float CurrentValue = SplineComp->GetFloatPropertyAtSplineInputKey(NearestParam, SelectionState->GetSelectedAttributeName());
+	const int32 NewAttributeIndex = SplineComp->SetFloatPropertyAtSplineInputKey(NearestParam, CurrentValue, SelectionState->GetSelectedAttributeName());
 
 	// Update selection
 	SelectionState->SetSelectedSegmentIndex(INDEX_NONE);
