SplineComponentVisualizer: Cache spline elements for rendering

This prevents unnecessary spline sampling.

diff --git a/Engine/Source/Editor/ComponentVisualizers/Private/SplineComponentVisualizer.cpp b/Engine/Source/Editor/ComponentVisualizers/Private/SplineComponentVisualizer.cpp
--- a/Engine/Source/Editor/ComponentVisualizers/Private/SplineComponentVisualizer.cpp
+++ b/Engine/Source/Editor/ComponentVisualizers/Private/SplineComponentVisualizer.cpp
@@ -57,6 +57,13 @@ DEFINE_LOG_CATEGORY_STATIC(LogSplineComponentVisualizer, Log, All)
 
 namespace SplineComponentVisualizerLocals
 {
+	bool GRebuildPDICacheEveryFrame = false;
+	FAutoConsoleVariableRef CVarUseSplineCurves(
+		TEXT("SplineComponentVisualizer.RebuildPDICacheEveryFrame"),
+		GRebuildPDICacheEveryFrame,
+		TEXT("When true, the spline is fully resampled for rendering every frame. Otherwise, the cache is invalidated only when necessary.")
+	);
+	
 	bool IsEnabledForSpline(const USplineComponent* InSplineComponent)
 	{
 		auto ModularFeatureName = ISplineComponentVisualizerSuppressor::GetModularFeatureName();
@@ -664,265 +671,327 @@ void FSplineComponentVisualizer::AddReferencedObjects(FReferenceCollector& Colle
 	}
 }
 
-static double GetDashSize(const FSceneView* View, const FVector& Start, const FVector& End, float Scale)
+void FSplineComponentVisualizer::DirtyPDICache(const USplineComponent* SplineComp)
+{
+	FPDICache& ComponentPDICache = PDICache.FindOrAdd(SplineComp);
+	ComponentPDICache.bDirty = true;
+}
+
+void FSplineComponentVisualizer::UpdatePDICache(const USplineComponent* SplineComp)
 {
-	const double StartW = View->WorldToScreen(Start).W;
-	const double EndW = View->WorldToScreen(End).W;
+	FPDICache& ComponentPDICache = PDICache.FindOrAdd(SplineComp);
 
-	const double WLimit = 10.0f;
-	if (StartW > WLimit || EndW > WLimit)
+	if (!(SplineComponentVisualizerLocals::GRebuildPDICacheEveryFrame || ComponentPDICache.bDirty))
 	{
-		return FMath::Max(StartW, EndW) * Scale;
+		return;
 	}
+	
+	FPDICache::FLineBatch LineBatch;
+	FPDICache::FPointBatch PointBatch;
+	ComponentPDICache.Reset();
+	
+	const FInterpCurveVector& SplineInfo = SplineComp->GetSplinePointsPosition();
+	const USplineComponent* EditedSplineComp = GetEditedSplineComponent();
 
-	return 0.0f;
-}
+	const bool bIsSplineEditable = !SplineComp->bModifiedByConstructionScript; // bSplineHasBeenEdited || SplineInfo == Archetype->SplineCurves.Position || SplineComp->bInputSplinePointsToConstructionScript;
 
-void FSplineComponentVisualizer::DrawVisualization(const UActorComponent* Component, const FSceneView* View, FPrimitiveDrawInterface* PDI)
-{
-	if (const USplineComponent* SplineComp = Cast<const USplineComponent>(Component))
+	constexpr FColor ReadOnlyColor = FColor(255, 0, 255, 255);
+	const FColor NormalColor = bIsSplineEditable ? FColor(SplineComp->EditorUnselectedSplineSegmentColor.ToFColor(true)) : ReadOnlyColor;
+	const FColor SelectedColor = bIsSplineEditable ? FColor(SplineComp->EditorSelectedSplineSegmentColor.ToFColor(true)) : ReadOnlyColor;
+	const FColor TangentColor = bIsSplineEditable ? FColor(SplineComp->EditorTangentColor.ToFColor(true)) : ReadOnlyColor;
+	const float GrabHandleSize = 10.0f + (bIsSplineEditable ? GetDefault<ULevelEditorViewportSettings>()->SelectedSplinePointSizeAdjustment : 0.0f);
+
+	// Draw the tangent handles before anything else so they will not overdraw the rest of the spline
+	if (SplineComp == EditedSplineComp)
 	{
-		if (!SplineComponentVisualizerLocals::IsEnabledForSpline(SplineComp))
+		check(SelectionState);
+
+		if (SplineComp->GetNumberOfSplinePoints() == 0 && SelectionState->GetSelectedKeys().Num() > 0)
 		{
-			return;
+			ChangeSelectionState(INDEX_NONE, false);
 		}
-		
-		const FInterpCurveVector& SplineInfo = SplineComp->GetSplinePointsPosition();
-		const USplineComponent* EditedSplineComp = GetEditedSplineComponent();
-
-		const USplineComponent* Archetype = CastChecked<USplineComponent>(SplineComp->GetArchetype());
-		const bool bIsSplineEditable = !SplineComp->bModifiedByConstructionScript; // bSplineHasBeenEdited || SplineInfo == Archetype->SplineCurves.Position || SplineComp->bInputSplinePointsToConstructionScript;
-
-		const FColor ReadOnlyColor = FColor(255, 0, 255, 255);
-		const FColor NormalColor = bIsSplineEditable ? FColor(SplineComp->EditorUnselectedSplineSegmentColor.ToFColor(true)) : ReadOnlyColor;
-		const FColor SelectedColor = bIsSplineEditable ? FColor(SplineComp->EditorSelectedSplineSegmentColor.ToFColor(true)) : ReadOnlyColor;
-		const FColor TangentColor = bIsSplineEditable ? FColor(SplineComp->EditorTangentColor.ToFColor(true)) : ReadOnlyColor;
-		const float GrabHandleSize = 10.0f + (bIsSplineEditable ? GetDefault<ULevelEditorViewportSettings>()->SelectedSplinePointSizeAdjustment : 0.0f);
-
-		// Draw the tangent handles before anything else so they will not overdraw the rest of the spline
-		if (SplineComp == EditedSplineComp)
+		else
 		{
-			check(SelectionState);
-
-			if (SplineComp->GetNumberOfSplinePoints() == 0 && SelectionState->GetSelectedKeys().Num() > 0)
-			{
-				ChangeSelectionState(INDEX_NONE, false);
-			}
-			else
+			for (const TSet<int32> SelectedKeysCopy = SelectionState->GetSelectedKeys();
+				int32 SelectedKey : SelectedKeysCopy)
 			{
-				const TSet<int32> SelectedKeysCopy = SelectionState->GetSelectedKeys();
-				for (int32 SelectedKey : SelectedKeysCopy)
+				check(SelectedKey >= 0);
+				if (SelectedKey >= SplineComp->GetNumberOfSplinePoints())
+				{
+					// Catch any keys that might not exist anymore due to the underlying component changing.
+					ChangeSelectionState(SelectedKey, true);
+					continue;
+				}
+
+				if (SplineInfo.Points[SelectedKey].IsCurveKey())
 				{
-					check(SelectedKey >= 0);
-					if (SelectedKey >= SplineComp->GetNumberOfSplinePoints())
+					const float TangentHandleSize = 8.0f + (bIsSplineEditable ? GetDefault<ULevelEditorViewportSettings>()->SplineTangentHandleSizeAdjustment : 0.0f);
+					const float TangentScale = GetDefault<ULevelEditorViewportSettings>()->SplineTangentScale;
+
+					const FVector Location = SplineComp->GetLocationAtSplinePoint(SelectedKey, ESplineCoordinateSpace::World);
+					const FVector LeaveTangent = SplineComp->GetLeaveTangentAtSplinePoint(SelectedKey, ESplineCoordinateSpace::World) * TangentScale;
+					const FVector ArriveTangent = SplineComp->GetArriveTangentAtSplinePoint(SelectedKey, ESplineCoordinateSpace::World) * TangentScale;
+					
+					// determine tangent coloration
+					const bool bTangentSelected = (SelectedKey == SelectionState->GetSelectedTangentHandle());
+					const ESelectedTangentHandle SelectedTangentHandleType = SelectionState->GetSelectedTangentHandleType();
+					const bool bArriveSelected = bTangentSelected && (SelectedTangentHandleType == ESelectedTangentHandle::Arrive);
+					const bool bLeaveSelected = bTangentSelected && (SelectedTangentHandleType == ESelectedTangentHandle::Leave);
+					FColor ArriveColor = bArriveSelected ? SelectedColor : TangentColor;
+					FColor LeaveColor = bLeaveSelected ? SelectedColor : TangentColor;
+
+					LineBatch.Elements.Reset(2);
+					LineBatch.Elements.Emplace(Location, Location - ArriveTangent, ArriveColor, SDPG_Foreground);
+					LineBatch.Elements.Emplace(Location, Location + LeaveTangent, LeaveColor, SDPG_Foreground);
+					ComponentPDICache.AddBatch(MoveTemp(LineBatch));
+					
+					PointBatch.Elements.Reset(1);
+					if (bIsSplineEditable)
 					{
-						// Catch any keys that might not exist anymore due to the underlying component changing.
-						ChangeSelectionState(SelectedKey, true);
-						continue;
+						PointBatch.AllocHitProxyFunc = [SplineComp, SelectedKey]()
+							{ return new HSplineTangentHandleProxy(SplineComp, SelectedKey, false); };
 					}
+					PointBatch.Elements.Emplace(Location + LeaveTangent, LeaveColor, TangentHandleSize, SDPG_Foreground);
+					ComponentPDICache.AddBatch(MoveTemp(PointBatch));
 
-					if (SplineInfo.Points[SelectedKey].IsCurveKey())
+					PointBatch.Elements.Reset(1);
+					if (bIsSplineEditable)
 					{
-						const float TangentHandleSize = 8.0f + (bIsSplineEditable ? GetDefault<ULevelEditorViewportSettings>()->SplineTangentHandleSizeAdjustment : 0.0f);
-						const float TangentScale = GetDefault<ULevelEditorViewportSettings>()->SplineTangentScale;
-
-						const FVector Location = SplineComp->GetLocationAtSplinePoint(SelectedKey, ESplineCoordinateSpace::World);
-						const FVector LeaveTangent = SplineComp->GetLeaveTangentAtSplinePoint(SelectedKey, ESplineCoordinateSpace::World) * TangentScale;
-						const FVector ArriveTangent = SplineComp->GetArriveTangentAtSplinePoint(SelectedKey, ESplineCoordinateSpace::World) * TangentScale;
-
-						PDI->SetHitProxy(NULL);
-
-						// determine tangent coloration
-						const bool bTangentSelected = (SelectedKey == SelectionState->GetSelectedTangentHandle());
-						const ESelectedTangentHandle SelectedTangentHandleType = SelectionState->GetSelectedTangentHandleType();
-						const bool bArriveSelected = bTangentSelected && (SelectedTangentHandleType == ESelectedTangentHandle::Arrive);
-						const bool bLeaveSelected = bTangentSelected && (SelectedTangentHandleType == ESelectedTangentHandle::Leave);
-						FColor ArriveColor = bArriveSelected ? SelectedColor : TangentColor;
-						FColor LeaveColor = bLeaveSelected ? SelectedColor : TangentColor;
-						
-						PDI->DrawLine(Location, Location - ArriveTangent, ArriveColor, SDPG_Foreground);
-						PDI->DrawLine(Location, Location + LeaveTangent, LeaveColor, SDPG_Foreground);
-
-						if (bIsSplineEditable)
-						{
-							PDI->SetHitProxy(new HSplineTangentHandleProxy(Component, SelectedKey, false));
-						}
-						PDI->DrawPoint(Location + LeaveTangent, LeaveColor, TangentHandleSize, SDPG_Foreground);
-
-						if (bIsSplineEditable)
-						{
-							PDI->SetHitProxy(new HSplineTangentHandleProxy(Component, SelectedKey, true));
-						}
-						PDI->DrawPoint(Location - ArriveTangent, ArriveColor, TangentHandleSize, SDPG_Foreground);
-
-						PDI->SetHitProxy(NULL);
+						PointBatch.AllocHitProxyFunc = [SplineComp, SelectedKey]()
+							{ return new HSplineTangentHandleProxy(SplineComp, SelectedKey, true); };
 					}
+					PointBatch.Elements.Emplace(Location - ArriveTangent, ArriveColor, TangentHandleSize, SDPG_Foreground);
+					ComponentPDICache.AddBatch(MoveTemp(PointBatch));
 				}
 			}
 		}
+	}
 
-		const bool bShouldVisualizeScale = SplineComp->bShouldVisualizeScale;
-		const float DefaultScale = SplineComp->ScaleVisualizationWidth;
+	const bool bShouldVisualizeScale = SplineComp->bShouldVisualizeScale;
+	const float DefaultScale = SplineComp->ScaleVisualizationWidth;
 
-		FVector OldKeyPos(0);
-		FVector OldKeyRightVector(0);
-		FVector OldKeyScale(0);
+	FVector OldKeyPos(0);
+	FVector OldKeyRightVector(0);
+	FVector OldKeyScale(0);
 
-		const TSet<int32>& SelectedKeys = SelectionState->GetSelectedKeys();
+	const TSet<int32>& SelectedKeys = SelectionState->GetSelectedKeys();
 
-		const int32 NumPoints = SplineInfo.Points.Num();
-		const int32 NumSegments = SplineInfo.bIsLooped ? NumPoints : NumPoints - 1;
-		for (int32 KeyIdx = 0; KeyIdx < NumSegments + 1; KeyIdx++)
-		{
-			const FVector NewKeyPos = SplineComp->GetLocationAtSplinePoint(KeyIdx, ESplineCoordinateSpace::World);
-			const FVector NewKeyRightVector = SplineComp->GetRightVectorAtSplinePoint(KeyIdx, ESplineCoordinateSpace::World);
-			const FVector NewKeyUpVector = SplineComp->GetUpVectorAtSplinePoint(KeyIdx, ESplineCoordinateSpace::World);
-			const FVector NewKeyScale = SplineComp->GetScaleAtSplinePoint(KeyIdx) * DefaultScale;
+	const int32 NumPoints = SplineInfo.Points.Num();
+	const int32 NumSegments = SplineInfo.bIsLooped ? NumPoints : NumPoints - 1;
+	for (int32 KeyIdx = 0; KeyIdx < NumSegments + 1; KeyIdx++)
+	{
+		const FVector NewKeyPos = SplineComp->GetLocationAtSplinePoint(KeyIdx, ESplineCoordinateSpace::World);
+		const FVector NewKeyRightVector = SplineComp->GetRightVectorAtSplinePoint(KeyIdx, ESplineCoordinateSpace::World);
+		const FVector NewKeyUpVector = SplineComp->GetUpVectorAtSplinePoint(KeyIdx, ESplineCoordinateSpace::World);
+		const FVector NewKeyScale = SplineComp->GetScaleAtSplinePoint(KeyIdx) * DefaultScale;
 
-			const FColor KeyColor = (SplineComp == EditedSplineComp && SelectedKeys.Contains(KeyIdx)) ? SelectedColor : NormalColor;
+		const FColor KeyColor = (SplineComp == EditedSplineComp && SelectedKeys.Contains(KeyIdx)) ? SelectedColor : NormalColor;
 
-			// Draw the keypoint and up/right vectors
-			if (KeyIdx < NumPoints)
+		// Draw the keypoint and up/right vectors
+		if (KeyIdx < NumPoints)
+		{
+			if (bShouldVisualizeScale)
 			{
-				if (bShouldVisualizeScale)
+				constexpr int32 ArcPoints = 20;
+
+				LineBatch.Elements.Reset(3 + ArcPoints);
+				LineBatch.Elements.Emplace(NewKeyPos, NewKeyPos - NewKeyRightVector * NewKeyScale.Y, KeyColor, SDPG_Foreground);
+				LineBatch.Elements.Emplace(NewKeyPos, NewKeyPos + NewKeyRightVector * NewKeyScale.Y, KeyColor, SDPG_Foreground);
+				LineBatch.Elements.Emplace(NewKeyPos, NewKeyPos + NewKeyUpVector * NewKeyScale.Z, KeyColor, SDPG_Foreground);
+				FVector OldArcPos = NewKeyPos + NewKeyRightVector * NewKeyScale.Y;
+				for (int32 ArcIndex = 1; ArcIndex <= ArcPoints; ArcIndex++)
 				{
-					PDI->SetHitProxy(NULL);
+					float Sin;
+					float Cos;
+					FMath::SinCos(&Sin, &Cos, static_cast<float>(ArcIndex) * PI / static_cast<float>(ArcPoints));
+					const FVector NewArcPos = NewKeyPos + Cos * NewKeyRightVector * NewKeyScale.Y + Sin * NewKeyUpVector * NewKeyScale.Z;
+					LineBatch.Elements.Emplace(OldArcPos, NewArcPos, KeyColor, SDPG_Foreground);
+					OldArcPos = NewArcPos;
+				}
+				ComponentPDICache.AddBatch(MoveTemp(LineBatch));
+			}
 
-					PDI->DrawLine(NewKeyPos, NewKeyPos - NewKeyRightVector * NewKeyScale.Y, KeyColor, SDPG_Foreground);
-					PDI->DrawLine(NewKeyPos, NewKeyPos + NewKeyRightVector * NewKeyScale.Y, KeyColor, SDPG_Foreground);
-					PDI->DrawLine(NewKeyPos, NewKeyPos + NewKeyUpVector * NewKeyScale.Z, KeyColor, SDPG_Foreground);
+			PointBatch.Elements.Reset(1);
+			if (bIsSplineEditable)
+			{
+				PointBatch.AllocHitProxyFunc = [SplineComp, KeyIdx]()
+					{  return new HSplineKeyProxy(SplineComp, KeyIdx); };
+			}
+			PointBatch.Elements.Emplace(NewKeyPos, KeyColor, GrabHandleSize, SDPG_Foreground);
+			ComponentPDICache.AddBatch(MoveTemp(PointBatch));
+		}
+
+		// If not the first keypoint, draw a line to the previous keypoint.
+		if (KeyIdx > 0)
+		{
+			// how many?
+			LineBatch.Elements.Reset();
+			
+			const FColor LineColor = NormalColor;
+			if (bIsSplineEditable)
+			{
+				LineBatch.AllocHitProxyFunc = [SplineComp, KeyIdx]()
+					{ return new HSplineSegmentProxy(SplineComp, KeyIdx - 1); };
+			}
 
-					const int32 ArcPoints = 20;
-					FVector OldArcPos = NewKeyPos + NewKeyRightVector * NewKeyScale.Y;
-					for (int32 ArcIndex = 1; ArcIndex <= ArcPoints; ArcIndex++)
+			const float SegmentLineThickness = GetDefault<ULevelEditorViewportSettings>()->SplineLineThicknessAdjustment;
+			constexpr int32 NumStepsPerSegment = 21;
+			
+			// For constant interpolation - don't draw ticks - just draw dotted line.
+			if (SplineInfo.Points[KeyIdx - 1].InterpMode == CIM_Constant)
+			{
+				FVector OldPos = OldKeyPos;
+						
+				const FVector SegmentVector = SplineComp->GetLocationAtSplinePoint(KeyIdx, ESplineCoordinateSpace::World) - OldKeyPos;
+				const float SegmentLength = static_cast<float>(SegmentVector.Length());
+				const FVector SegmentDirection = SegmentVector.GetSafeNormal();
+				const float StepLength = SegmentLength / NumStepsPerSegment;
+				const FVector SegmentStep = SegmentDirection * StepLength;
+						
+				for (int32 StepIdx = 1; StepIdx <= NumStepsPerSegment; StepIdx++)
+				{
+					const FVector NewPos = OldPos + SegmentStep;
+					
+					if (StepIdx % 2 == 1)
 					{
-						float Sin;
-						float Cos;
-						FMath::SinCos(&Sin, &Cos, static_cast<float>(ArcIndex) * PI / static_cast<float>(ArcPoints));
-						const FVector NewArcPos = NewKeyPos + Cos * NewKeyRightVector * NewKeyScale.Y + Sin * NewKeyUpVector * NewKeyScale.Z;
-						PDI->DrawLine(OldArcPos, NewArcPos, KeyColor, SDPG_Foreground);
-						OldArcPos = NewArcPos;
+						LineBatch.Elements.Emplace(OldPos, NewPos, LineColor, SDPG_World, SegmentLineThickness);
 					}
-				}
 
-				if (bIsSplineEditable)
-				{
-					PDI->SetHitProxy(new HSplineKeyProxy(Component, KeyIdx));
+					OldPos = NewPos;
 				}
-				PDI->DrawPoint(NewKeyPos, KeyColor, GrabHandleSize, SDPG_Foreground);
-				PDI->SetHitProxy(NULL);
 			}
-
-			// If not the first keypoint, draw a line to the previous keypoint.
-			if (KeyIdx > 0)
+			else
 			{
-				const FColor LineColor = NormalColor;
-				if (bIsSplineEditable)
-				{
-					PDI->SetHitProxy(new HSplineSegmentProxy(Component, KeyIdx - 1));
-				}
+				// Determine the colors to use
+				const bool bIsEdited = (SplineComp == EditedSplineComp);
+				const bool bKeyIdxLooped = (SplineInfo.bIsLooped && KeyIdx == NumPoints);
+				const int32 BeginIdx = bKeyIdxLooped ? 0 : KeyIdx;
+				const int32 EndIdx = KeyIdx - 1;
+				const bool bBeginSelected = SelectedKeys.Contains(BeginIdx);
+				const bool bEndSelected = SelectedKeys.Contains(EndIdx);
+				const FColor BeginColor = bIsEdited && bBeginSelected ? SelectedColor : NormalColor;
+				const FColor EndColor = bIsEdited && bEndSelected ? SelectedColor : NormalColor;
+				
+				// Find position on first keyframe.
+				FVector OldPos = OldKeyPos;
+				FVector OldRightVector = OldKeyRightVector;
+				FVector OldScale = OldKeyScale;
 
-				// For constant interpolation - don't draw ticks - just draw dotted line.
-				if (SplineInfo.Points[KeyIdx - 1].InterpMode == CIM_Constant)
+				// Then draw a line for each substep.
+				constexpr float PartialGradientProportion = 0.75f;
+				constexpr int32 PartialNumSteps = static_cast<int32>(NumStepsPerSegment * PartialGradientProportion);
+
+				for (int32 StepIdx = 1; StepIdx <= NumStepsPerSegment; StepIdx++)
 				{
-					const double DashSize = GetDashSize(View, OldKeyPos, NewKeyPos, 0.03f);
-					if (DashSize > 0.0f)
+					const float StepRatio = static_cast<float>(StepIdx) / static_cast<float>(NumStepsPerSegment);
+					const float Key = static_cast<float>(EndIdx) + StepRatio;
+					const FVector NewPos = SplineComp->GetLocationAtSplineInputKey(Key, ESplineCoordinateSpace::World);
+					const FVector NewRightVector = SplineComp->GetRightVectorAtSplineInputKey(Key, ESplineCoordinateSpace::World);
+					const FVector NewScale = SplineComp->GetScaleAtSplineInputKey(Key) * DefaultScale;
+
+					// creates a gradient that starts partway through the selection
+					FColor StepColor;
+					if (bBeginSelected == bEndSelected)
 					{
-						DrawDashedLine(PDI, OldKeyPos, NewKeyPos, LineColor, DashSize, SDPG_World);
+						StepColor = BeginColor;
 					}
-				}
-				else
-				{
-					// Determine the colors to use
-					const bool bIsEdited = (SplineComp == EditedSplineComp);
-					const bool bKeyIdxLooped = (SplineInfo.bIsLooped && KeyIdx == NumPoints);
-					const int32 BeginIdx = bKeyIdxLooped ? 0 : KeyIdx;
-					const int32 EndIdx = KeyIdx - 1;
-					const bool bBeginSelected = SelectedKeys.Contains(BeginIdx);
-					const bool bEndSelected = SelectedKeys.Contains(EndIdx);
-					const FColor BeginColor = (bIsEdited && bBeginSelected) ? SelectedColor : NormalColor;
-					const FColor EndColor = (bIsEdited && bEndSelected) ? SelectedColor : NormalColor;
-					
-					// Find position on first keyframe.
-					FVector OldPos = OldKeyPos;
-					FVector OldRightVector = OldKeyRightVector;
-					FVector OldScale = OldKeyScale;
-
-					// Then draw a line for each substep.
-					constexpr int32 NumSteps = 20;
-					constexpr float PartialGradientProportion = 0.75f;
-					constexpr int32 PartialNumSteps = (int32)(NumSteps * PartialGradientProportion);
-					const float SegmentLineThickness = GetDefault<ULevelEditorViewportSettings>()->SplineLineThicknessAdjustment;
-
-					for (int32 StepIdx = 1; StepIdx <= NumSteps; StepIdx++)
+					else if (bBeginSelected && StepIdx > (NumStepsPerSegment - PartialNumSteps))
 					{
-						const float StepRatio = static_cast<float>(StepIdx) / static_cast<float>(NumSteps);
-						const float Key = static_cast<float>(EndIdx) + StepRatio;
-						const FVector NewPos = SplineComp->GetLocationAtSplineInputKey(Key, ESplineCoordinateSpace::World);
-						const FVector NewRightVector = SplineComp->GetRightVectorAtSplineInputKey(Key, ESplineCoordinateSpace::World);
-						const FVector NewScale = SplineComp->GetScaleAtSplineInputKey(Key) * DefaultScale;
-
-						// creates a gradient that starts partway through the selection
-						FColor StepColor;
-						if (bBeginSelected == bEndSelected)
-						{
-							StepColor = BeginColor;
-						}
-						else if (bBeginSelected && StepIdx > (NumSteps - PartialNumSteps))
-						{
-							const float LerpRatio = (1.0f - StepRatio) / PartialGradientProportion;
-							StepColor = FMath::Lerp(BeginColor.ReinterpretAsLinear(), EndColor.ReinterpretAsLinear(), LerpRatio).ToFColor(false);
-						}
-						else if (bEndSelected && StepIdx <= PartialNumSteps)
-						{
-							const float LerpRatio = 1.0f - (StepRatio / PartialGradientProportion);
-							StepColor = FMath::Lerp(BeginColor.ReinterpretAsLinear(), EndColor.ReinterpretAsLinear(), LerpRatio).ToFColor(false);
-						}
-						else
-						{
-							StepColor = NormalColor; // unselected
-						}
-						
-						PDI->DrawLine(OldPos, NewPos, StepColor, SDPG_Foreground, SegmentLineThickness);
-						if (bShouldVisualizeScale)
-						{
-							PDI->DrawLine(OldPos - OldRightVector * OldScale.Y, NewPos - NewRightVector * NewScale.Y, LineColor, SDPG_Foreground);
-							PDI->DrawLine(OldPos + OldRightVector * OldScale.Y, NewPos + NewRightVector * NewScale.Y, LineColor, SDPG_Foreground);
-
-							#if VISUALIZE_SPLINE_UPVECTORS
-							const FVector NewUpVector = SplineComp->GetUpVectorAtSplineInputKey(Key, ESplineCoordinateSpace::World);
-							PDI->DrawLine(NewPos, NewPos + NewUpVector * SplineComp->ScaleVisualizationWidth * 0.5f, LineColor, SDPG_Foreground);
-							PDI->DrawLine(NewPos, NewPos + NewRightVector * SplineComp->ScaleVisualizationWidth * 0.5f, LineColor, SDPG_Foreground);
-							#endif
-						}
-
-						OldPos = NewPos;
-						OldRightVector = NewRightVector;
-						OldScale = NewScale;
+						const float LerpRatio = (1.0f - StepRatio) / PartialGradientProportion;
+						StepColor = FMath::Lerp(BeginColor.ReinterpretAsLinear(), EndColor.ReinterpretAsLinear(), LerpRatio).ToFColor(false);
+					}
+					else if (bEndSelected && StepIdx <= PartialNumSteps)
+					{
+						const float LerpRatio = 1.0f - (StepRatio / PartialGradientProportion);
+						StepColor = FMath::Lerp(BeginColor.ReinterpretAsLinear(), EndColor.ReinterpretAsLinear(), LerpRatio).ToFColor(false);
+					}
+					else
+					{
+						StepColor = NormalColor; // unselected
+					}
+
+					LineBatch.Elements.Emplace(OldPos, NewPos, StepColor, SDPG_Foreground, SegmentLineThickness);
+					if (bShouldVisualizeScale)
+					{
+						LineBatch.Elements.Emplace(OldPos - OldRightVector * OldScale.Y, NewPos - NewRightVector * NewScale.Y, LineColor, SDPG_Foreground);
+						LineBatch.Elements.Emplace(OldPos + OldRightVector * OldScale.Y, NewPos + NewRightVector * NewScale.Y, LineColor, SDPG_Foreground);
+
+						#if VISUALIZE_SPLINE_UPVECTORS
+						const FVector NewUpVector = SplineComp->GetUpVectorAtSplineInputKey(Key, ESplineCoordinateSpace::World);
+						LineBatch.Elements.Emplace(NewPos, NewPos + NewUpVector * SplineComp->ScaleVisualizationWidth * 0.5f, LineColor, SDPG_Foreground);
+						LineBatch.Elements.Emplace(NewPos, NewPos + NewRightVector * SplineComp->ScaleVisualizationWidth * 0.5f, LineColor, SDPG_Foreground);
+						#endif
 					}
-				}
 
-				PDI->SetHitProxy(NULL);
+					OldPos = NewPos;
+					OldRightVector = NewRightVector;
+					OldScale = NewScale;
+				}
 			}
 
-			OldKeyPos = NewKeyPos;
-			OldKeyRightVector = NewKeyRightVector;
-			OldKeyScale = NewKeyScale;
+			ComponentPDICache.AddBatch(MoveTemp(LineBatch));
 		}
 
-		// Render the selected attribute channel's values.
-		if (FName SelectedAttributeName = SelectionState->GetSelectedAttributeName(); SelectedAttributeName != NAME_None)
+		OldKeyPos = NewKeyPos;
+		OldKeyRightVector = NewKeyRightVector;
+		OldKeyScale = NewKeyScale;
+	}
+
+	// Render the selected attribute channel's values.
+	if (FName SelectedAttributeName = SelectionState->GetSelectedAttributeName(); SelectedAttributeName != NAME_None)
+	{
+		for (int Idx = 0; Idx < SplineComp->GetNumberOfPropertyValues(SelectedAttributeName); ++Idx)
 		{
-			for (int Idx = 0; Idx < SplineComp->GetNumberOfPropertyValues(SelectedAttributeName); ++Idx)
-			{
-				PDI->SetHitProxy(new HSplineAttributeKeyProxy(Component, Idx));
-				
-				float Parameter = SplineComp->GetFloatPropertyInputKeyAtIndex(Idx, SelectedAttributeName);
-				FVector PositionOnCurve = SplineComp->GetLocationAtSplineInputKey(Parameter, ESplineCoordinateSpace::World);
-				PDI->DrawPoint(PositionOnCurve, FLinearColor::Red, GrabHandleSize, SDPG_Foreground);
-			}
-					
-			PDI->SetHitProxy(NULL);
+			PointBatch.Elements.Reset(1);
+
+			PointBatch.AllocHitProxyFunc = [SplineComp, Idx]()
+				{ return new HSplineAttributeKeyProxy(SplineComp, Idx); };
+			
+			float Parameter = SplineComp->GetFloatPropertyInputKeyAtIndex(Idx, SelectedAttributeName);
+			FVector PositionOnCurve = SplineComp->GetLocationAtSplineInputKey(Parameter, ESplineCoordinateSpace::World);
+			PointBatch.Elements.Emplace(PositionOnCurve, FColor::Red, GrabHandleSize, SDPG_Foreground);
+
+			ComponentPDICache.AddBatch(MoveTemp(PointBatch));
+		}
+	}
+
+	ComponentPDICache.bDirty = false;
+}
+
+void FSplineComponentVisualizer::DrawVisualization(const UActorComponent* Component, const FSceneView* View, FPrimitiveDrawInterface* PDI)
+{
+	// TODO: DROP SUBDIVISION COUNT DURING INTERACTIVE CHANGES TO REDUCE RESAMPLE TIME AND GIVE AN IDEA OF THE LOOK
+	
+	QUICK_SCOPE_CYCLE_COUNTER(STAT_SplineComponentVisualizer_DrawVisualization);
+	
+	if (const USplineComponent* SplineComp = Cast<const USplineComponent>(Component))
+	{
+		if (!SplineComponentVisualizerLocals::IsEnabledForSpline(SplineComp))
+		{
+			return;
+		}
+
+		// If the component is not in the cache, this is our first time seeing it.
+		if (!PDICache.Find(SplineComp))
+		{
+			// Not ideal to const_cast here, but 
+			USplineComponent* MutableSplineComponent = const_cast<USplineComponent*>(SplineComp);
+
+			// We assume it is safe to capture 'this' and 'MutableSplineComponent' because the lambda's user object will be 'this' and the delegates
+			// are owned by the component itself (and therefore cannot be broadcast if the component is not valid).
+			auto InvalidateCacheFunc = [this, MutableSplineComponent]() { DirtyPDICache(MutableSplineComponent); };
+			
+			MutableSplineComponent->GetOnSplineChanged().AddSPLambda(this, InvalidateCacheFunc);
+			MutableSplineComponent->GetOnSplineUpdated().AddSPLambda(this, InvalidateCacheFunc);
+			MutableSplineComponent->GetOnSplineDisplayChanged().AddSPLambda(this, InvalidateCacheFunc);
+		}
+
+		// This does nothing if the cache is not dirty
+		UpdatePDICache(SplineComp);
+
+		if (const FPDICache* ComponentCache = PDICache.Find(SplineComp))
+		{
+			ComponentCache->Draw(PDI);
 		}
 	}
 }
@@ -1045,6 +1114,11 @@ void FSplineComponentVisualizer::ChangeSelectionState(int32 Index, bool bIsCtrlH
 			DeselectedInEditorDelegateHandle = GetEditedSplineComponent()->OnDeselectedInEditor.AddRaw(this, &FSplineComponentVisualizer::OnDeselectedInEditor);
 		}
 	}
+
+	if (const USplineComponent* SplineComp = GetEditedSplineComponent())
+	{
+		DirtyPDICache(SplineComp);
+	}
 }
 
 
@@ -1054,6 +1128,8 @@ const USplineComponent* FSplineComponentVisualizer::UpdateSelectedSplineComponen
 
 	const USplineComponent* SplineComp = CastChecked<const USplineComponent>(VisProxy->Component.Get());
 
+	DirtyPDICache(SplineComp);
+	
 	AActor* OldSplineOwningActor = SelectionState->GetSplinePropertyPath().GetParentOwningActor();
 	FComponentPropertyPath NewSplinePropertyPath(SplineComp);
 	SelectionState->SetSplinePropertyPath(NewSplinePropertyPath);
@@ -1250,6 +1326,11 @@ bool FSplineComponentVisualizer::VisProxyHandleClick(FEditorViewportClient* InVi
 
 	if (bVisProxyClickHandled)
 	{
+		if (const USplineComponent* SplineComp = GetEditedSplineComponent())
+		{
+			DirtyPDICache(SplineComp);
+		}
+		
 		GEditor->RedrawLevelEditingViewports(true);
 	}
 
@@ -1264,6 +1345,8 @@ void FSplineComponentVisualizer::SetEditedSplineComponent(const USplineComponent
 
 	FComponentPropertyPath SplinePropertyPath(InSplineComponent);
 	SelectionState->SetSplinePropertyPath(SplinePropertyPath);
+	
+	DirtyPDICache(InSplineComponent);
 }
 
 USplineComponent* FSplineComponentVisualizer::GetEditedSplineComponent() const
@@ -1420,109 +1503,110 @@ bool FSplineComponentVisualizer::HandleInputDelta(FEditorViewportClient* Viewpor
 	ResetTempModes();
 
 	USplineComponent* SplineComp = GetEditedSplineComponent();
-	bool bInputHandled = false;
-	if (SplineComp != nullptr)
+	if (!SplineComp || !SplineComponentVisualizerLocals::IsEnabledForSpline(SplineComp))
 	{
-		if (!SplineComponentVisualizerLocals::IsEnabledForSpline(SplineComp))
-		{
-			return false;
-		}
+		return false;
+	}
+	
+	bool bInputHandled = false;
 		
-		if (IsAnySelectedKeyIndexOutOfRange(SplineComp))
-		{
-			// Something external has changed the number of spline points, meaning that the cached selected keys are no longer valid
-			EndEditing();
-			return false;
-		}
+	if (IsAnySelectedKeyIndexOutOfRange(SplineComp))
+	{
+		// Something external has changed the number of spline points, meaning that the cached selected keys are no longer valid
+		EndEditing();
+		return false;
+	}
 
-		check(SelectionState);
+	check(SelectionState);
 
-		// Use a local value for DeltaTranslate so that we can modify it based on the SnapToSurface setting without
-		// changing it for the caller (that parameter should probably be const, but it's a base class function)
-		FVector DeltaTranslate = DeltaTranslateIn;
-
-		// It's tough to port the actor "surface snapping" toggle behavior here because the original code only works on actors
-		// and involves things like keeping track of last transform, snapping to surface even if moving in a different plane, etc
-		// (see FLevelEditorViewportClient::ProjectActorsIntoWorld). However we can at least support the limited but useful case
-		// of trying to drag points around by the ball of the widget.
-		const ULevelEditorViewportSettings* ViewportSettings = GetDefault<ULevelEditorViewportSettings>();
-		const float SnapOffsetExtent = (ViewportSettings->SnapToSurface.bEnabled) ? (ViewportSettings->SnapToSurface.SnapOffsetExtent) : (0.0f);
-		if (ViewportSettings->SnapToSurface.bEnabled
-			&& ViewportClient->GetWidgetMode() == UE::Widget::EWidgetMode::WM_Translate
-			&& (ViewportClient->GetCurrentWidgetAxis() == EAxisList::Screen || ViewportClient->GetCurrentWidgetAxis() == EAxisList::XYZ)
-			&& !ViewportClient->IsOrtho())
+	// Use a local value for DeltaTranslate so that we can modify it based on the SnapToSurface setting without
+	// changing it for the caller (that parameter should probably be const, but it's a base class function)
+	FVector DeltaTranslate = DeltaTranslateIn;
+
+	// It's tough to port the actor "surface snapping" toggle behavior here because the original code only works on actors
+	// and involves things like keeping track of last transform, snapping to surface even if moving in a different plane, etc
+	// (see FLevelEditorViewportClient::ProjectActorsIntoWorld). However we can at least support the limited but useful case
+	// of trying to drag points around by the ball of the widget.
+	const ULevelEditorViewportSettings* ViewportSettings = GetDefault<ULevelEditorViewportSettings>();
+	const float SnapOffsetExtent = (ViewportSettings->SnapToSurface.bEnabled) ? (ViewportSettings->SnapToSurface.SnapOffsetExtent) : (0.0f);
+	if (ViewportSettings->SnapToSurface.bEnabled
+		&& ViewportClient->GetWidgetMode() == UE::Widget::EWidgetMode::WM_Translate
+		&& (ViewportClient->GetCurrentWidgetAxis() == EAxisList::Screen || ViewportClient->GetCurrentWidgetAxis() == EAxisList::XYZ)
+		&& !ViewportClient->IsOrtho())
+	{
+		FHitResult HitResult;
+		if (RaycastWorld(SplineComp->GetWorld(), ViewportClient, Viewport, HitResult))
 		{
-			FHitResult HitResult;
-			if (RaycastWorld(SplineComp->GetWorld(), ViewportClient, Viewport, HitResult))
-			{
-				FVector3d NewGizmoLocation = HitResult.ImpactPoint + ViewportSettings->SnapToSurface.SnapOffsetExtent * HitResult.ImpactNormal;
-				DeltaTranslate = NewGizmoLocation - ViewportClient->GetWidgetLocation();
-			}
+			FVector3d NewGizmoLocation = HitResult.ImpactPoint + ViewportSettings->SnapToSurface.SnapOffsetExtent * HitResult.ImpactNormal;
+			DeltaTranslate = NewGizmoLocation - ViewportClient->GetWidgetLocation();
 		}
+	}
 
-		if (SelectionState->GetSelectedAttributeIndex() != INDEX_NONE)
-		{
-			bInputHandled = TransformSelectedAttribute(EPropertyChangeType::Interactive, DeltaTranslate);
-		}
-		else if (SelectionState->GetSelectedTangentHandle() != INDEX_NONE)
-		{
-			// Transform the tangent using an EPropertyChangeType::Interactive change. Later on, at the end of mouse tracking, a non-interactive change will be notified via void TrackingStopped :
-			bInputHandled = TransformSelectedTangent(EPropertyChangeType::Interactive, DeltaTranslate);
-		}
-		else if (ViewportClient->IsAltPressed())
+	if (SelectionState->GetSelectedAttributeIndex() != INDEX_NONE)
+	{
+		bInputHandled = TransformSelectedAttribute(EPropertyChangeType::Interactive, DeltaTranslate);
+	}
+	else if (SelectionState->GetSelectedTangentHandle() != INDEX_NONE)
+	{
+		// Transform the tangent using an EPropertyChangeType::Interactive change. Later on, at the end of mouse tracking, a non-interactive change will be notified via void TrackingStopped :
+		bInputHandled = TransformSelectedTangent(EPropertyChangeType::Interactive, DeltaTranslate);
+	}
+	else if (ViewportClient->IsAltPressed())
+	{
+		if (ViewportClient->GetWidgetMode() == UE::Widget::WM_Translate && ViewportClient->GetCurrentWidgetAxis() != EAxisList::None && SelectionState->GetSelectedKeys().Num() == 1)
 		{
-			if (ViewportClient->GetWidgetMode() == UE::Widget::WM_Translate && ViewportClient->GetCurrentWidgetAxis() != EAxisList::None && SelectionState->GetSelectedKeys().Num() == 1)
-			{
-				static const int MaxDuplicationDelay = 3;
+			static const int MaxDuplicationDelay = 3;
 
-				FVector Drag = DeltaTranslate;
+			FVector Drag = DeltaTranslate;
 
-				if (bAllowDuplication)
+			if (bAllowDuplication)
+			{
+				float SmallestGridSize = 1.0f;
+				const TArray<float>& PosGridSizes = GEditor->GetCurrentPositionGridArray();
+				if (PosGridSizes.IsValidIndex(0))
 				{
-					float SmallestGridSize = 1.0f;
-					const TArray<float>& PosGridSizes = GEditor->GetCurrentPositionGridArray();
-					if (PosGridSizes.IsValidIndex(0))
-					{
-						SmallestGridSize = PosGridSizes[0];
-					}
+					SmallestGridSize = PosGridSizes[0];
+				}
 
-					// When grid size is set to a value other than the smallest grid size, do not delay duplication
-					if (DuplicateDelay >= MaxDuplicationDelay || GEditor->GetGridSize() > SmallestGridSize)
-					{
-						Drag += DuplicateDelayAccumulatedDrag;
-						DuplicateDelayAccumulatedDrag = FVector::ZeroVector;
+				// When grid size is set to a value other than the smallest grid size, do not delay duplication
+				if (DuplicateDelay >= MaxDuplicationDelay || GEditor->GetGridSize() > SmallestGridSize)
+				{
+					Drag += DuplicateDelayAccumulatedDrag;
+					DuplicateDelayAccumulatedDrag = FVector::ZeroVector;
 
-						bAllowDuplication = false;
-						bDuplicatingSplineKey = true;
+					bAllowDuplication = false;
+					bDuplicatingSplineKey = true;
 
-						DuplicateKeyForAltDrag(Drag);
-					}
-					else
-					{ 
-						DuplicateDelay++;
-						DuplicateDelayAccumulatedDrag += DeltaTranslate;
-					}
+					DuplicateKeyForAltDrag(Drag);
 				}
 				else
-				{
-					UpdateDuplicateKeyForAltDrag(Drag);
+				{ 
+					DuplicateDelay++;
+					DuplicateDelayAccumulatedDrag += DeltaTranslate;
 				}
-
-				bInputHandled = true;
 			}
-		}
-		else
-		{
-			// Transform the spline keys using an EPropertyChangeType::Interactive change. Later on, at the end of mouse tracking, a non-interactive change will be notified via void TrackingStopped :
-			bInputHandled = TransformSelectedKeys(EPropertyChangeType::Interactive, DeltaTranslate, DeltaRotate, DeltaScale);
+			else
+			{
+				UpdateDuplicateKeyForAltDrag(Drag);
+			}
+
+			bInputHandled = true;
 		}
 	}
+	else
+	{
+		// Transform the spline keys using an EPropertyChangeType::Interactive change. Later on, at the end of mouse tracking, a non-interactive change will be notified via void TrackingStopped :
+		bInputHandled = TransformSelectedKeys(EPropertyChangeType::Interactive, DeltaTranslate, DeltaRotate, DeltaScale);
+	}
+	
 	if (bInputHandled)
 	{
 		if (AActor* OwnerActor = SplineComp->GetOwner())
 		{
 			OwnerActor->PostEditMove(false);
 		}
+
+		DirtyPDICache(SplineComp);
 	}
 
 	return bInputHandled;
diff --git a/Engine/Source/Editor/ComponentVisualizers/Public/SplineComponentVisualizer.h b/Engine/Source/Editor/ComponentVisualizers/Public/SplineComponentVisualizer.h
--- a/Engine/Source/Editor/ComponentVisualizers/Public/SplineComponentVisualizer.h
+++ b/Engine/Source/Editor/ComponentVisualizers/Public/SplineComponentVisualizer.h
@@ -28,6 +28,7 @@
 #include "UObject/ObjectMacros.h"
 #include "UObject/UObjectGlobals.h"
 #include "Features/IModularFeature.h"
+#include "PrimitiveDrawInterface.h"
 
 #include "SplineComponentVisualizer.generated.h"
 
@@ -51,6 +52,128 @@ class USplineMetadata;
 struct FConvexVolume;
 struct FViewportClick;
 
+namespace UE::SplineComponentVisualizer
+{
+	
+struct FPDICache
+{
+	struct FRenderableLine
+	{
+		FRenderableLine()
+			: Start(ForceInitToZero)
+			, End(ForceInitToZero)
+			, Color(ForceInitToZero)
+			, DepthPriority(SDPG_World)
+			, Thickness(0.0f)
+			, DepthBias(0.0f)
+		{}
+
+		FRenderableLine(const FVector& InStart, const FVector& InEnd, const FColor& InColor, const ESceneDepthPriorityGroup InDepthPriority = SDPG_World, const float InThickness = 0.0f, const float InDepthBias = 0.0f)
+			: Start(InStart)
+			, End(InEnd)
+			, Color(InColor)
+			, DepthPriority(InDepthPriority)
+			, Thickness(InThickness)
+			, DepthBias(InDepthBias)
+		{}
+
+		void Draw(FPrimitiveDrawInterface* PDI) const
+		{
+			PDI->DrawLine(Start, End, Color, static_cast<uint8>(DepthPriority), Thickness, DepthBias);
+		}
+
+		FVector Start;
+		FVector End;
+		FColor Color;
+		ESceneDepthPriorityGroup DepthPriority;
+		float Thickness;
+		float DepthBias;
+	};
+	
+	struct FRenderablePoint
+	{
+		FRenderablePoint()
+			: Position(ForceInitToZero)
+			, Color(ForceInitToZero)
+			, Size(0.f)
+			, DepthPriority(SDPG_World)
+		{}
+
+		FRenderablePoint(const FVector& InPosition, const FColor& InColor, const float InSize, const ESceneDepthPriorityGroup InDepthPriority = SDPG_World)
+			: Position(InPosition)
+			, Color(InColor)
+			, Size(InSize)
+			, DepthPriority(InDepthPriority)
+		{}
+
+		void Draw(FPrimitiveDrawInterface* PDI) const
+		{
+			PDI->DrawPoint(Position, Color, Size, static_cast<uint8>(DepthPriority));
+		}
+		
+		FVector Position;
+		FColor Color;
+		float Size;
+		ESceneDepthPriorityGroup DepthPriority;
+	};
+
+	template <typename ElementType>
+	struct TElementBatch
+	{
+		TArray<ElementType> Elements;
+		TFunction<HHitProxy*()> AllocHitProxyFunc = []() -> HHitProxy* { return nullptr; };
+
+		void Draw(FPrimitiveDrawInterface* PDI) const
+		{
+			PDI->SetHitProxy(AllocateHitProxy());
+			for (const ElementType& Element : Elements) { Element.Draw(PDI); }
+			PDI->SetHitProxy(nullptr);
+		}
+		
+	private:
+		
+		HHitProxy* AllocateHitProxy() const
+		{
+			return AllocHitProxyFunc
+				? AllocHitProxyFunc()
+				: nullptr;
+		}
+	};
+	
+	using FLineBatch = TElementBatch<FRenderableLine>;
+	using FPointBatch = TElementBatch<FRenderablePoint>;
+
+	void Reset()
+	{
+		LineBatches.Reset();
+		PointBatches.Reset();
+	}
+	
+	void Draw(FPrimitiveDrawInterface* PDI) const
+	{
+		for (const FLineBatch& Batch : LineBatches) { Batch.Draw(PDI); }
+		for (const FPointBatch& Batch : PointBatches) { Batch.Draw(PDI); }
+	}
+	
+	void AddBatch(FLineBatch&& LineBatch)
+	{
+		LineBatches.Add(MoveTemp(LineBatch));
+		LineBatch.Elements.Reset();
+	}
+
+	void AddBatch(FPointBatch&& PointBatch)
+	{
+		PointBatches.Add(MoveTemp(PointBatch));
+		PointBatch.Elements.Reset();
+	}
+	
+	TArray<FLineBatch> LineBatches;
+	TArray<FPointBatch> PointBatches;
+	bool bDirty = true;
+};
+	
+}
+
 /** Tangent handle selection modes. */
 UENUM()
 enum class ESelectedTangentHandle
@@ -263,7 +386,10 @@ PRAGMA_DISABLE_INTERNAL_WARNINGS
 	, public ISplineDetailsProvider
 PRAGMA_DISABLE_INTERNAL_WARNINGS
 {
+	using FPDICache = UE::SplineComponentVisualizer::FPDICache;
+	
 public:
+
 	UE_API FSplineComponentVisualizer();
 	UE_API virtual ~FSplineComponentVisualizer();
 
@@ -559,6 +685,11 @@ protected:
 	mutable TArray<TSharedPtr<FName>> SharedAttributeNames;
 	
 private:
+
+	TMap<TWeakObjectPtr<const USplineComponent>, FPDICache> PDICache;
+	void DirtyPDICache(const USplineComponent* SplineComp);
+	void UpdatePDICache(const USplineComponent* SplineComp);
+	
 	TSharedPtr<SSplineGeneratorPanel> SplineGeneratorPanel;
 	static UE_API TWeakPtr<SWindow> WeakExistingWindow;
 };
