SplineComponentVisualizer: Fix problematic implementation of ISplineDetailsProvider::ShouldUseForSpline

A recent change decoupling the spline component visualizer from the spline component details customization caused a regression in some of the functionality of the details customization.

diff --git a/Engine/Source/Editor/ComponentVisualizers/Private/SplineComponentVisualizer.cpp b/Engine/Source/Editor/ComponentVisualizers/Private/SplineComponentVisualizer.cpp
--- a/Engine/Source/Editor/ComponentVisualizers/Private/SplineComponentVisualizer.cpp
+++ b/Engine/Source/Editor/ComponentVisualizers/Private/SplineComponentVisualizer.cpp
@@ -4391,14 +4391,8 @@ void FSplineComponentVisualizer::OnDeselectedInEditor(TObjectPtr<USplineComponen
 
 bool FSplineComponentVisualizer::ShouldUseForSpline(const USplineComponent* SplineComponent) const
 {
-	if (!GetEditedSplineComponent())
-	{
-		return false;
-	}
-	
 	// This implementation should be used when the classic visualizer is enabled for this spline
-	return (GetEditedSplineComponent() == SplineComponent || GetEditedSplineComponent()->GetArchetype() == SplineComponent)
-		&& SplineComponentVisualizerLocals::IsEnabledForSpline(SplineComponent);
+	return SplineComponentVisualizerLocals::IsEnabledForSpline(SplineComponent);
 }
 
 bool FSplineComponentVisualizer::IsEnabledForSpline(const USplineComponent* InSplineComponent) const
