SplineComponentVisualizer: Dirty PDI cache at proper times

Dirty when:
1) Moving the component
2) Changing the current attribute channel

diff --git a/Engine/Source/Editor/ComponentVisualizers/Private/SplineComponentVisualizer.cpp b/Engine/Source/Editor/ComponentVisualizers/Private/SplineComponentVisualizer.cpp
--- a/Engine/Source/Editor/ComponentVisualizers/Private/SplineComponentVisualizer.cpp
+++ b/Engine/Source/Editor/ComponentVisualizers/Private/SplineComponentVisualizer.cpp
@@ -981,6 +981,7 @@ void FSplineComponentVisualizer::DrawVisualization(const UActorComponent* Compon
 			// are owned by the component itself (and therefore cannot be broadcast if the component is not valid).
 			auto InvalidateCacheFunc = [this, MutableSplineComponent]() { DirtyPDICache(MutableSplineComponent); };
 			
+			MutableSplineComponent->TransformUpdated.AddSPLambda(this, [InvalidateCacheFunc](USceneComponent*, EUpdateTransformFlags, ETeleportType) { InvalidateCacheFunc(); });
 			MutableSplineComponent->GetOnSplineChanged().AddSPLambda(this, InvalidateCacheFunc);
 			MutableSplineComponent->GetOnSplineUpdated().AddSPLambda(this, InvalidateCacheFunc);
 			MutableSplineComponent->GetOnSplineDisplayChanged().AddSPLambda(this, InvalidateCacheFunc);
@@ -4303,6 +4304,11 @@ TSharedPtr<SWidget> FSplineComponentVisualizer::GenerateChannelSelectorWidget()
 	{
 		SelectionState->SetSelectedAttribute(INDEX_NONE, SelectedItem ? *SelectedItem : NAME_None);
 
+		for (TPair<TWeakObjectPtr<const USplineComponent>, FPDICache>& ComponentPDICache : PDICache)
+		{
+			ComponentPDICache.Value.bDirty = true;
+		}
+			
 		if (TSharedPtr<SComboButton> PinnedComboButton = WeakCombo.Pin())
 		{
 			PinnedComboButton->SetIsOpen(false);
diff --git a/Engine/Source/Editor/ComponentVisualizers/Public/SplineComponentVisualizer.h b/Engine/Source/Editor/ComponentVisualizers/Public/SplineComponentVisualizer.h
--- a/Engine/Source/Editor/ComponentVisualizers/Public/SplineComponentVisualizer.h
+++ b/Engine/Source/Editor/ComponentVisualizers/Public/SplineComponentVisualizer.h
@@ -686,7 +686,7 @@ protected:
 	
 private:
 
-	TMap<TWeakObjectPtr<const USplineComponent>, FPDICache> PDICache;
+	mutable TMap<TWeakObjectPtr<const USplineComponent>, FPDICache> PDICache;
 	void DirtyPDICache(const USplineComponent* SplineComp);
 	void UpdatePDICache(const USplineComponent* SplineComp);
 	
