Fix crash in FSplineComponentVisualizer::GenerateSplinePointTypeSubMenu()

diff --git a/Engine/Source/Editor/ComponentVisualizers/Private/SplineComponentVisualizer.cpp b/Engine/Source/Editor/ComponentVisualizers/Private/SplineComponentVisualizer.cpp
--- a/Engine/Source/Editor/ComponentVisualizers/Private/SplineComponentVisualizer.cpp
+++ b/Engine/Source/Editor/ComponentVisualizers/Private/SplineComponentVisualizer.cpp
@@ -3923,20 +3923,21 @@ void FSplineComponentVisualizer::GenerateSelectSplinePointsSubMenu(FMenuBuilder&
 }
 void FSplineComponentVisualizer::GenerateSplinePointTypeSubMenu(FMenuBuilder& MenuBuilder) const
 {
-	const USplineComponent* SplineComp = GetEditedSplineComponent();
-	check(SplineComp);
-	const TArray<ESplinePointType::Type> EnabledSplinePointTypes = SplineComp->GetEnabledSplinePointTypes();
-	if (EnabledSplinePointTypes.Contains(ESplinePointType::Curve))
-	{
-		MenuBuilder.AddMenuEntry(FSplineComponentVisualizerCommands::Get().SetKeyToCurve);
-	}
-	if (EnabledSplinePointTypes.Contains(ESplinePointType::Linear))
-	{
-		MenuBuilder.AddMenuEntry(FSplineComponentVisualizerCommands::Get().SetKeyToLinear);
-	}
-	if (EnabledSplinePointTypes.Contains(ESplinePointType::Constant))
+	if (const USplineComponent* SplineComp = GetEditedSplineComponent())
 	{
-		MenuBuilder.AddMenuEntry(FSplineComponentVisualizerCommands::Get().SetKeyToConstant);
+		const TArray<ESplinePointType::Type> EnabledSplinePointTypes = SplineComp->GetEnabledSplinePointTypes();
+		if (EnabledSplinePointTypes.Contains(ESplinePointType::Curve))
+		{
+			MenuBuilder.AddMenuEntry(FSplineComponentVisualizerCommands::Get().SetKeyToCurve);
+		}
+		if (EnabledSplinePointTypes.Contains(ESplinePointType::Linear))
+		{
+			MenuBuilder.AddMenuEntry(FSplineComponentVisualizerCommands::Get().SetKeyToLinear);
+		}
+		if (EnabledSplinePointTypes.Contains(ESplinePointType::Constant))
+		{
+			MenuBuilder.AddMenuEntry(FSplineComponentVisualizerCommands::Get().SetKeyToConstant);
+		}
 	}
 }
 
