SplineComponentVisualizer: Fix ensure caused by spline point duplication

diff --git a/Engine/Source/Editor/ComponentVisualizers/Private/SplineComponentVisualizer.cpp b/Engine/Source/Editor/ComponentVisualizers/Private/SplineComponentVisualizer.cpp
--- a/Engine/Source/Editor/ComponentVisualizers/Private/SplineComponentVisualizer.cpp
+++ b/Engine/Source/Editor/ComponentVisualizers/Private/SplineComponentVisualizer.cpp
@@ -2700,7 +2700,14 @@ void FSplineComponentVisualizer::OnDuplicateKey()
 		check(SelectedKeyIndex >= 0);
 		check(SelectedKeyIndex < SplineComp->GetNumberOfSplinePoints());
 
-		SplineComp->AddPoint(SplineComp->GetSplinePointAt(SelectedKeyIndex, ESplineCoordinateSpace::Local), false);
+		const int32 InsertIndex = SelectedKeyIndex + 1;
+		
+		FSplinePoint NewPoint = SplineComp->GetSplinePointAt(SelectedKeyIndex, ESplineCoordinateSpace::Local);
+		SplineComp->AddSplinePointAtIndex(NewPoint.Position, InsertIndex, ESplineCoordinateSpace::Local, false);
+		SplineComp->SetScaleAtSplinePoint(InsertIndex, NewPoint.Scale, false);
+		SplineComp->SetQuaternionAtSplinePoint(InsertIndex, NewPoint.Rotation.Quaternion(), ESplineCoordinateSpace::Local, false);
+		SplineComp->SetTangentsAtSplinePoint(InsertIndex, NewPoint.ArriveTangent, NewPoint.LeaveTangent, ESplineCoordinateSpace::Local, false);
+		SplineComp->SetSplinePointType(InsertIndex, NewPoint.Type, false);
 	}
 
 	SelectionState->Modify();
