InterpCurve crash fix for SplineComponent

diff --git a/Engine/Source/Editor/ComponentVisualizers/Private/SplineComponentVisualizer.cpp b/Engine/Source/Editor/ComponentVisualizers/Private/SplineComponentVisualizer.cpp
--- a/Engine/Source/Editor/ComponentVisualizers/Private/SplineComponentVisualizer.cpp
+++ b/Engine/Source/Editor/ComponentVisualizers/Private/SplineComponentVisualizer.cpp
@@ -2774,7 +2774,9 @@ void FSplineComponentVisualizer::OnAddKeyToSegment()
 	check(SelectionState->GetSelectedTangentHandleType() == ESelectedTangentHandle::None);
 
 	SelectionState->Modify();
-
+	
+	SplineComp->UpdateSpline();
+	
 	SplitSegment(SelectionState->GetSelectedSplinePosition(), SelectionState->GetSelectedSegmentIndex());
 
 	SelectionState->SetSelectedSegmentIndex(INDEX_NONE);
@@ -2793,6 +2795,8 @@ bool FSplineComponentVisualizer::DuplicateKeyForAltDrag(const FVector& InDrag)
 	check(SelectedKeys.Num() == 1);
 	check(SelectedKeys.Contains(LastKeyIndexSelected));
 
+	SplineComp->UpdateSpline();
+	
 	// When dragging from end point, maximum angle is 60 degrees from attached segment
 	// to determine whether to split existing segment or create a new point
 	static const double Angle60 = 1.0472;
