SplineComponentVisualizer: Fix details integration for spline subclasses with their own visualizer subclass

We take into account the spline class when determining if a visualizer instance ought to be considered a valid details provider for a given spline.

diff --git a/Engine/Source/Editor/ComponentVisualizers/Private/SplineComponentVisualizer.cpp b/Engine/Source/Editor/ComponentVisualizers/Private/SplineComponentVisualizer.cpp
--- a/Engine/Source/Editor/ComponentVisualizers/Private/SplineComponentVisualizer.cpp
+++ b/Engine/Source/Editor/ComponentVisualizers/Private/SplineComponentVisualizer.cpp
@@ -4406,11 +4406,10 @@ void FSplineComponentVisualizer::OnDeselectedInEditor(TObjectPtr<USplineComponen
 	EndEditing();
 }
 
-
 bool FSplineComponentVisualizer::ShouldUseForSpline(const USplineComponent* SplineComponent) const
 {
-	// This implementation should be used when the classic visualizer is enabled for this spline
-	return SplineComponentVisualizerLocals::IsEnabledForSpline(SplineComponent);
+	// This implementation should be used when the classic visualizer is enabled for this spline and this is the registered visualizer for the class of the specified spline.
+	return GUnrealEd->FindComponentVisualizer(SplineComponent->GetClass()).Get() == this && SplineComponentVisualizerLocals::IsEnabledForSpline(SplineComponent);
 }
 
 bool FSplineComponentVisualizer::IsEnabledForSpline(const USplineComponent* InSplineComponent) const
