SplineComponentVisualizer: Fix incorrect metadata when duplicating spline points

diff --git a/Engine/Source/Editor/ComponentVisualizers/Private/SplineComponentVisualizer.cpp b/Engine/Source/Editor/ComponentVisualizers/Private/SplineComponentVisualizer.cpp
--- a/Engine/Source/Editor/ComponentVisualizers/Private/SplineComponentVisualizer.cpp
+++ b/Engine/Source/Editor/ComponentVisualizers/Private/SplineComponentVisualizer.cpp
@@ -2708,6 +2708,13 @@ void FSplineComponentVisualizer::OnDuplicateKey()
 		SplineComp->SetQuaternionAtSplinePoint(InsertIndex, NewPoint.Rotation.Quaternion(), ESplineCoordinateSpace::Local, false);
 		SplineComp->SetTangentsAtSplinePoint(InsertIndex, NewPoint.ArriveTangent, NewPoint.LeaveTangent, ESplineCoordinateSpace::Local, false);
 		SplineComp->SetSplinePointType(InsertIndex, NewPoint.Type, false);
+
+		// USplineComponent::AddSplinePointAtIndex inserts a default metadata point, but the expectation is that the metadata point is duplicated.
+		// This can be implemented by overwriting the inserted point with the metadata point at the selected index.
+		if (USplineMetadata* Metadata = SplineComp->GetSplinePointsMetadata())
+		{
+			Metadata->CopyPoint(Metadata, SelectedKeyIndex, InsertIndex);
+		}
 	}
 
 	SelectionState->Modify();
