[Bug Fix] Fixed Localization bug

Fixed localization bug cause by not using UI_COMMAND or FUICommandInfo::MakeCommandInfo for dynamic content
Other functions are not gathered for localization

#Synced-CL 48346153

diff --git a/Engine/Source/Editor/ComponentVisualizers/Private/SplineComponentVisualizer.cpp b/Engine/Source/Editor/ComponentVisualizers/Private/SplineComponentVisualizer.cpp
--- a/Engine/Source/Editor/ComponentVisualizers/Private/SplineComponentVisualizer.cpp
+++ b/Engine/Source/Editor/ComponentVisualizers/Private/SplineComponentVisualizer.cpp
@@ -233,12 +233,16 @@ public:
 		UI_COMMAND(ToggleSnapTangentAdjustments, "Allow Tangents Updates On Snap", "Allow tangents to update when performing snap operations on points.", EUserInterfaceActionType::ToggleButton, FInputChord());
 
 #define AXIS_UI_COMMAND(CommandName, Axis, LabelText, TooltipText, ButtonType) \
-	CommandName = FUICommandInfoDecl(this->AsShared(), \
+	FUICommandInfo::MakeCommandInfo( \
+		AsShared(), \
+		CommandName, \
 		TEXT(#CommandName), \
 		FText::Format(LOCTEXT(#CommandName, LabelText), AxisDisplayInfo::GetAxisDisplayName(Axis)), \
-		FText::Format(LOCTEXT(#CommandName TEXT("_ToolTip"), TooltipText), AxisDisplayInfo::GetAxisDisplayName(Axis))) \
-		.UserInterfaceType(ButtonType) \
-		.DefaultChord(FInputChord());
+		FText::Format(LOCTEXT(#CommandName TEXT("_ToolTip"), TooltipText), AxisDisplayInfo::GetAxisDisplayName(Axis)), \
+		FSlateIcon(), \
+		ButtonType, \
+		FInputChord() \
+		); \
 
 		AXIS_UI_COMMAND(SnapAllToSelectedY, EAxisList::Left, "Snap All To Selected {0}", "Snap all spline points to selected spline point world {0} position.", EUserInterfaceActionType::Button)
 		AXIS_UI_COMMAND(SnapAllToSelectedZ, EAxisList::Up, "Snap All To Selected {0}", "Snap all spline points to selected spline point world {0} position.", EUserInterfaceActionType::Button)
diff --git a/Engine/Source/Editor/LevelEditor/Private/LevelEditorActions.cpp b/Engine/Source/Editor/LevelEditor/Private/LevelEditorActions.cpp
--- a/Engine/Source/Editor/LevelEditor/Private/LevelEditorActions.cpp
+++ b/Engine/Source/Editor/LevelEditor/Private/LevelEditorActions.cpp
@@ -4329,12 +4329,16 @@ void FLevelEditorCommands::RegisterCommands()
 	UI_COMMAND( DeltaTransformToActors, "Delta Transform", "Apply Delta Transform to selected elements", EUserInterfaceActionType::Button, FInputChord() );
 
 	#define AXIS_UI_COMMAND(CommandName, Axis) \
-	CommandName = FUICommandInfoDecl(this->AsShared(), \
+	FUICommandInfo::MakeCommandInfo( \
+		AsShared(), \
+		CommandName, \
 		TEXT(#CommandName), \
 		FText::Format(LOCTEXT("MirrorAxisLabel", "Mirror {0} Axis"), AxisDisplayInfo::GetAxisDisplayName(Axis)), \
-		FText::Format(LOCTEXT("MirrorAxisTooltip", "Mirrors the element along the {0} axis"), AxisDisplayInfo::GetAxisDisplayName(Axis))) \
-		.UserInterfaceType(EUserInterfaceActionType::Button) \
-		.DefaultChord(FInputChord());
+		FText::Format(LOCTEXT("MirrorAxisTooltip", "Mirrors the element along the {0} axis"), AxisDisplayInfo::GetAxisDisplayName(Axis)), \
+		FSlateIcon(), \
+		EUserInterfaceActionType::Button, \
+		FInputChord() \
+		); \
 
 	AXIS_UI_COMMAND(MirrorActorY, EAxisList::Left)
 	AXIS_UI_COMMAND(MirrorActorZ, EAxisList::Up)
