[Splines] When deleting spline points, select the next point down, instead of always the first point.
This improves the experience of deleting points from the end or middle of the spline.

#rnx

#Synced-CL 50104405

diff --git a/Engine/Source/Editor/ComponentVisualizers/Private/SplineComponentVisualizer.cpp b/Engine/Source/Editor/ComponentVisualizers/Private/SplineComponentVisualizer.cpp
--- a/Engine/Source/Editor/ComponentVisualizers/Private/SplineComponentVisualizer.cpp
+++ b/Engine/Source/Editor/ComponentVisualizers/Private/SplineComponentVisualizer.cpp
@@ -3228,8 +3228,9 @@ void FSplineComponentVisualizer::OnDeleteKey()
 		SplineComp->RemoveSplinePoint(SelectedKeyIndex, false);
 	}
 
-	// Select first key
-	ChangeSelectionState(0, false);
+	// Select the first valid key before the deleted range
+	const int32 LastKeyIndex = SelectedKeysSorted.Last();
+	ChangeSelectionState(LastKeyIndex > 1 ? LastKeyIndex - 1 : 0, false);
 	SelectionState->Modify();
 	SelectionState->ClearSelectedSegmentIndex();
 	SelectionState->ClearSelectedTangentHandle();
