SplineComponentVisualizer: Clear PDI cache on selection change

#Synced-CL 50600896

diff --git a/Engine/Source/Editor/ComponentVisualizers/Private/SplineComponentVisualizer.cpp b/Engine/Source/Editor/ComponentVisualizers/Private/SplineComponentVisualizer.cpp
--- a/Engine/Source/Editor/ComponentVisualizers/Private/SplineComponentVisualizer.cpp
+++ b/Engine/Source/Editor/ComponentVisualizers/Private/SplineComponentVisualizer.cpp
@@ -418,6 +418,16 @@ PRAGMA_ENABLE_DEPRECATION_WARNINGS
 	SelectionState = NewObject<USplineComponentVisualizerSelectionState>(GetTransientPackage(), TEXT("SelectionState"), RF_Transactional);
 
 	IModularFeatures::Get().RegisterModularFeature(ISplineDetailsProvider::GetModularFeatureName(), this);
+
+	OnObjectTransactedHandle = FCoreUObjectDelegates::OnObjectTransacted.AddLambda([this](UObject* Object, const class FTransactionObjectEvent& Transaction)
+	{
+		if (Object == SelectionState && Transaction.GetEventType() == ETransactionObjectEventType::UndoRedo)
+		{
+			// Because the set of elements to draw is a function of a spline and a selection,
+			// changing the selection invalidates all caches, not just the cached elements for any particular spline.
+			PDICache.Empty();
+		}
+	});
 }
 
 void FSplineComponentVisualizer::OnRegister()
@@ -663,6 +673,9 @@ bool FSplineComponentVisualizer::ShouldShowForSelectedSubcomponents(const UActor
 
 FSplineComponentVisualizer::~FSplineComponentVisualizer()
 {
+	FCoreUObjectDelegates::OnObjectTransacted.Remove(OnObjectTransactedHandle);
+	OnObjectTransactedHandle.Reset();
+
 	IModularFeatures::Get().UnregisterModularFeature(ISplineDetailsProvider::GetModularFeatureName(), this);
 	FSplineComponentVisualizerCommands::Unregister();
 }
@@ -987,6 +1000,12 @@ void FSplineComponentVisualizer::DrawVisualization(const UActorComponent* Compon
 			// are owned by the component itself (and therefore cannot be broadcast if the component is not valid).
 			auto InvalidateCacheFunc = [this, MutableSplineComponent]() { DirtyPDICache(MutableSplineComponent); };
 			
+			// We don't retain the handles so lets just make sure we don't double register if our cache got cleared.
+			MutableSplineComponent->TransformUpdated.RemoveAll(this);
+			MutableSplineComponent->GetOnSplineChanged().RemoveAll(this);
+			MutableSplineComponent->GetOnSplineUpdated().RemoveAll(this);
+			MutableSplineComponent->GetOnSplineDisplayChanged().RemoveAll(this);
+
 			MutableSplineComponent->TransformUpdated.AddSPLambda(this, [InvalidateCacheFunc](USceneComponent*, EUpdateTransformFlags, ETeleportType) { InvalidateCacheFunc(); });
 			MutableSplineComponent->GetOnSplineChanged().AddSPLambda(this, InvalidateCacheFunc);
 			MutableSplineComponent->GetOnSplineUpdated().AddSPLambda(this, InvalidateCacheFunc);
diff --git a/Engine/Source/Editor/ComponentVisualizers/Public/SplineComponentVisualizer.h b/Engine/Source/Editor/ComponentVisualizers/Public/SplineComponentVisualizer.h
--- a/Engine/Source/Editor/ComponentVisualizers/Public/SplineComponentVisualizer.h
+++ b/Engine/Source/Editor/ComponentVisualizers/Public/SplineComponentVisualizer.h
@@ -686,6 +686,8 @@ protected:
 	
 private:
 
+	FDelegateHandle OnObjectTransactedHandle;
+
 	mutable TMap<TWeakObjectPtr<const USplineComponent>, FPDICache> PDICache;
 	void DirtyPDICache(const USplineComponent* SplineComp);
 	void UpdatePDICache(const USplineComponent* SplineComp);
